Copying //UE4/Dev-Sequencer to //UE4/Dev-Main (Source: //UE4/Dev-Sequencer @ 3946692)

diff --git a/Engine/Plugins/MovieScene/LevelSequenceEditor/Source/LevelSequenceEditor/Private/Misc/LevelSequenceEditorActorSpawner.h b/Engine/Plugins/MovieScene/LevelSequenceEditor/Source/LevelSequenceEditor/Private/Misc/LevelSequenceEditorActorSpawner.h
--- a/Engine/Plugins/MovieScene/LevelSequenceEditor/Source/LevelSequenceEditor/Private/Misc/LevelSequenceEditorActorSpawner.h
+++ b/Engine/Plugins/MovieScene/LevelSequenceEditor/Source/LevelSequenceEditor/Private/Misc/LevelSequenceEditorActorSpawner.h
@@ -14,7 +14,7 @@ public:
 	virtual bool IsEditor() const override { return true; }
 #if WITH_EDITOR
 	virtual TValueOrError<FNewSpawnable, FText> CreateNewSpawnableType(UObject& SourceObject, UMovieScene& OwnerMovieScene, UActorFactory* ActorFactory = nullptr) override;
-	virtual void SetupDefaultsForSpawnable(UObject* SpawnedObject, const FGuid& Guid, const FTransformData& TransformData, TSharedRef<ISequencer> Sequencer, USequencerSettings* Settings) override;
+	virtual void SetupDefaultsForSpawnable(UObject* SpawnedObject, const FGuid& Guid, const TOptional<FTransformData>& TransformData, TSharedRef<ISequencer> Sequencer, USequencerSettings* Settings) override;
 	virtual bool CanSetupDefaultsForSpawnable(UObject* SpawnedObject) const override;
 #endif	
 };
\ No newline at end of file
diff --git a/Engine/Plugins/MovieScene/LevelSequenceEditor/Source/LevelSequenceEditor/Private/Misc/LevelSequenceEditorSpawnRegister.cpp b/Engine/Plugins/MovieScene/LevelSequenceEditor/Source/LevelSequenceEditor/Private/Misc/LevelSequenceEditorSpawnRegister.cpp
--- a/Engine/Plugins/MovieScene/LevelSequenceEditor/Source/LevelSequenceEditor/Private/Misc/LevelSequenceEditorSpawnRegister.cpp
+++ b/Engine/Plugins/MovieScene/LevelSequenceEditor/Source/LevelSequenceEditor/Private/Misc/LevelSequenceEditorSpawnRegister.cpp
@@ -156,7 +156,16 @@ void FLevelSequenceEditorSpawnRegister::SaveDefaultSpawnableState(const FGuid& B
 	}
 
 	UMovieSceneSequence* Sequence = Sequencer->GetEvaluationTemplate().GetSequence(TemplateID);
+	if (!Sequence)
+	{
+		return;
+	}
+
 	UMovieScene* MovieScene = Sequence ? Sequence->GetMovieScene() : nullptr;
+	if (!MovieScene)
+	{
+		return;
+	}
 
 	FMovieSceneSpawnable* Spawnable = MovieScene->FindSpawnable(BindingId);
 
@@ -308,7 +317,7 @@ TValueOrError<FNewSpawnable, FText> FLevelSequenceEditorSpawnRegister::CreateNew
 	return MakeError(LOCTEXT("NoSpawnerFound", "No spawner found to create new spawnable type"));
 }
 
-void FLevelSequenceEditorSpawnRegister::SetupDefaultsForSpawnable(UObject* SpawnedObject, const FGuid& Guid, const FTransformData& TransformData, TSharedRef<ISequencer> Sequencer, USequencerSettings* Settings)
+void FLevelSequenceEditorSpawnRegister::SetupDefaultsForSpawnable(UObject* SpawnedObject, const FGuid& Guid, const TOptional<FTransformData>& TransformData, TSharedRef<ISequencer> Sequencer, USequencerSettings* Settings)
 {
 	for (TSharedPtr<IMovieSceneObjectSpawner> MovieSceneObjectSpawner : MovieSceneObjectSpawners)
 	{
@@ -320,17 +329,17 @@ void FLevelSequenceEditorSpawnRegister::SetupDefaultsForSpawnable(UObject* Spawn
 	}
 }
 
-void FLevelSequenceEditorSpawnRegister::HandleConvertPossessableToSpawnable(UObject* OldObject, IMovieScenePlayer& Player, FTransformData& OutTransformData)
+void FLevelSequenceEditorSpawnRegister::HandleConvertPossessableToSpawnable(UObject* OldObject, IMovieScenePlayer& Player, TOptional<FTransformData>& OutTransformData)
 {
 	// @TODO: this could probably be handed off to a spawner if we need anything else to be convertible between spawnable/posessable
 
 	AActor* OldActor = Cast<AActor>(OldObject);
 	if (OldActor)
 	{
-		OutTransformData.Translation = OldActor->GetActorLocation();
-		OutTransformData.Rotation = OldActor->GetActorRotation();
-		OutTransformData.Scale = OldActor->GetActorScale();
-		OutTransformData.bValid = true;
+		OutTransformData.Emplace();
+		OutTransformData->Translation = OldActor->GetActorLocation();
+		OutTransformData->Rotation = OldActor->GetActorRotation();
+		OutTransformData->Scale = OldActor->GetActorScale();
 
 		GEditor->SelectActor(OldActor, false, true);
 		UWorld* World = Cast<UWorld>(Player.GetPlaybackContext());
diff --git a/Engine/Plugins/MovieScene/LevelSequenceEditor/Source/LevelSequenceEditor/Private/Misc/LevelSequenceEditorSpawnRegister.h b/Engine/Plugins/MovieScene/LevelSequenceEditor/Source/LevelSequenceEditor/Private/Misc/LevelSequenceEditorSpawnRegister.h
--- a/Engine/Plugins/MovieScene/LevelSequenceEditor/Source/LevelSequenceEditor/Private/Misc/LevelSequenceEditorSpawnRegister.h
+++ b/Engine/Plugins/MovieScene/LevelSequenceEditor/Source/LevelSequenceEditor/Private/Misc/LevelSequenceEditorSpawnRegister.h
@@ -40,8 +40,8 @@ public:
 	virtual void SaveDefaultSpawnableState(FMovieSceneSpawnable& Spawnable, FMovieSceneSequenceIDRef TemplateID, IMovieScenePlayer& Player) override;
 #if WITH_EDITOR
 	virtual TValueOrError<FNewSpawnable, FText> CreateNewSpawnableType(UObject& SourceObject, UMovieScene& OwnerMovieScene, UActorFactory* ActorFactory = nullptr) override;
-	virtual void SetupDefaultsForSpawnable(UObject* SpawnedObject, const FGuid& Guid, const FTransformData& TransformData, TSharedRef<ISequencer> Sequencer, USequencerSettings* Settings) override;
-	virtual void HandleConvertPossessableToSpawnable(UObject* OldObject, IMovieScenePlayer& Player, FTransformData& OutTransformData) override;
+	virtual void SetupDefaultsForSpawnable(UObject* SpawnedObject, const FGuid& Guid, const TOptional<FTransformData>& TransformData, TSharedRef<ISequencer> Sequencer, USequencerSettings* Settings) override;
+	virtual void HandleConvertPossessableToSpawnable(UObject* OldObject, IMovieScenePlayer& Player, TOptional<FTransformData>& OutTransformData) override;
 	virtual bool CanConvertSpawnableToPossessable(FMovieSceneSpawnable& Spawnable) const override;
 #endif
 
diff --git a/Engine/Plugins/MovieScene/MatineeToLevelSequence/Source/MatineeToLevelSequence/Private/MatineeToLevelSequenceModule.cpp b/Engine/Plugins/MovieScene/MatineeToLevelSequence/Source/MatineeToLevelSequence/Private/MatineeToLevelSequenceModule.cpp
--- a/Engine/Plugins/MovieScene/MatineeToLevelSequence/Source/MatineeToLevelSequence/Private/MatineeToLevelSequenceModule.cpp
+++ b/Engine/Plugins/MovieScene/MatineeToLevelSequence/Source/MatineeToLevelSequence/Private/MatineeToLevelSequenceModule.cpp
@@ -72,6 +72,7 @@
 #include "Dialogs/Dialogs.h"
 #include "Framework/Notifications/NotificationManager.h"
 #include "Widgets/Notifications/SNotificationList.h"
+#include "MovieSceneTimeHelpers.h"
 
 #define LOCTEXT_NAMESPACE "MatineeToLevelSequence"
 
@@ -487,7 +488,7 @@ protected:
 				if (MatineeAnimControlTrack->GetNumKeyframes() != 0 && PossessableGuid.IsValid())
 				{
 					UMovieSceneSkeletalAnimationTrack* SkeletalAnimationTrack = NewMovieScene->AddTrack<UMovieSceneSkeletalAnimationTrack>(PossessableGuid);	
-					float EndPlaybackRange = NewMovieScene->GetPlaybackRange().GetUpperBoundValue();
+					FFrameNumber EndPlaybackRange = MovieScene::DiscreteExclusiveUpper(NewMovieScene->GetPlaybackRange());
 					FMatineeImportTools::CopyInterpAnimControlTrack(MatineeAnimControlTrack, SkeletalAnimationTrack, EndPlaybackRange);
 				}
 			}
@@ -697,7 +698,8 @@ protected:
 			MatineeActor->InitInterp();
 
 			// Set the length
-			NewMovieScene->SetPlaybackRange(0.0f, MatineeActor->MatineeData->InterpLength);
+			const int32 LengthInFrames = (MatineeActor->MatineeData->InterpLength * NewMovieScene->GetFrameResolution()).FrameNumber.Value;
+			NewMovieScene->SetPlaybackRange(FFrameNumber(0), LengthInFrames + 1);
 
 			// Convert the groups
 			for (int32 i=0; i<MatineeActor->GroupInst.Num(); ++i)
diff --git a/Engine/Plugins/Runtime/AudioCapture/Source/AudioCaptureEditor/Private/AudioRecordingManager.cpp b/Engine/Plugins/Runtime/AudioCapture/Source/AudioCaptureEditor/Private/AudioRecordingManager.cpp
--- a/Engine/Plugins/Runtime/AudioCapture/Source/AudioCaptureEditor/Private/AudioRecordingManager.cpp
+++ b/Engine/Plugins/Runtime/AudioCapture/Source/AudioCaptureEditor/Private/AudioRecordingManager.cpp
@@ -103,6 +103,13 @@ void FAudioRecordingManager::StartRecording(const FRecordingSettings& InSettings
 	RecordingSampleRate = Info.preferredSampleRate;
 	NumInputChannels = Info.inputChannels;	
 
+	// Only support mono and stereo mic inputs for now...
+	if (NumInputChannels != 1 && NumInputChannels != 2)
+	{
+		UE_LOG(LogMicManager, Warning, TEXT("Audio recording only supports mono and stereo mic input."));
+		return;
+	}
+
 	// Reserve enough space in our current recording buffer for 10 seconds of audio to prevent slowing down due to allocations
 	if (Settings.RecordingDuration > 0.0f)
 	{
diff --git a/Engine/Source/Developer/BlueprintCompilerCppBackend/Private/BlueprintCompilerCppBackendUMG.cpp b/Engine/Source/Developer/BlueprintCompilerCppBackend/Private/BlueprintCompilerCppBackendUMG.cpp
--- a/Engine/Source/Developer/BlueprintCompilerCppBackend/Private/BlueprintCompilerCppBackendUMG.cpp
+++ b/Engine/Source/Developer/BlueprintCompilerCppBackend/Private/BlueprintCompilerCppBackendUMG.cpp
@@ -130,9 +130,9 @@ bool FBackendHelperUMG::SpecialStructureConstructorUMG(const UStruct* Struct, co
 		if (OutResult)
 		{
 			const FSectionEvaluationData* SectionEvaluationData = reinterpret_cast<const FSectionEvaluationData*>(ValuePtr);
-			*OutResult = FString::Printf(TEXT("FSectionEvaluationData(%d, %s)")
+			*OutResult = FString::Printf(TEXT("FSectionEvaluationData(%d, %d)")
 				, SectionEvaluationData->ImplIndex
-				, *FEmitHelper::FloatToString(SectionEvaluationData->ForcedTime));
+				, SectionEvaluationData->ForcedTime.Value);
 		}
 		return true;
 	}
diff --git a/Engine/Source/Developer/LogVisualizer/LogVisualizer.Build.cs b/Engine/Source/Developer/LogVisualizer/LogVisualizer.Build.cs
--- a/Engine/Source/Developer/LogVisualizer/LogVisualizer.Build.cs
+++ b/Engine/Source/Developer/LogVisualizer/LogVisualizer.Build.cs
@@ -10,7 +10,6 @@ public class LogVisualizer : ModuleRules
 			new string[] {
 				"DesktopPlatform",
 				"MainFrame",
-                "SequencerWidgets",
 				"LevelEditor"
 			}
 		);
@@ -34,7 +33,6 @@ public class LogVisualizer : ModuleRules
                 "EditorStyle",
 				"Engine",
 				"UnrealEd",
-                "SequencerWidgets",
 			}
 		);
 
diff --git a/Engine/Source/Developer/LogVisualizer/Private/SVisualLogger.cpp b/Engine/Source/Developer/LogVisualizer/Private/SVisualLogger.cpp
--- a/Engine/Source/Developer/LogVisualizer/Private/SVisualLogger.cpp
+++ b/Engine/Source/Developer/LogVisualizer/Private/SVisualLogger.cpp
@@ -170,7 +170,7 @@ void SVisualLogger::Construct(const FArguments& InArgs, const TSharedRef<SDockTa
 	FLogVisualizer::Get().GetEvents().OnFiltersChanged.AddRaw(this, &SVisualLogger::OnFiltersChanged);
 	FLogVisualizer::Get().GetEvents().OnLogLineSelectionChanged = FOnLogLineSelectionChanged::CreateRaw(this, &SVisualLogger::OnLogLineSelectionChanged);
 	FLogVisualizer::Get().GetEvents().OnKeyboardEvent = FOnKeyboardEvent::CreateRaw(this, &SVisualLogger::OnKeyboaedRedirection);
-	FLogVisualizer::Get().GetTimeSliderController().Get()->GetTimeSliderArgs().OnScrubPositionChanged = FOnScrubPositionChanged::CreateRaw(this, &SVisualLogger::OnScrubPositionChanged);
+	FLogVisualizer::Get().GetTimeSliderController().Get()->GetTimeSliderArgs().OnScrubPositionChanged = FVisualLoggerTimeSliderArgs::FOnScrubPositionChanged::CreateRaw(this, &SVisualLogger::OnScrubPositionChanged);
 
 	FVisualLoggerDatabase::Get().GetEvents().OnRowSelectionChanged.AddRaw(this, &SVisualLogger::OnObjectSelectionChanged);
 	FVisualLoggerDatabase::Get().GetEvents().OnNewItem.AddRaw(this, &SVisualLogger::OnNewItemHandler);
@@ -734,7 +734,7 @@ void SVisualLogger::ResetData()
 
 	FLogVisualizer::Get().GetEvents().OnLogLineSelectionChanged = FOnLogLineSelectionChanged::CreateRaw(this, &SVisualLogger::OnLogLineSelectionChanged);
 	FLogVisualizer::Get().GetEvents().OnKeyboardEvent = FOnKeyboardEvent::CreateRaw(this, &SVisualLogger::OnKeyboaedRedirection);
-	FLogVisualizer::Get().GetTimeSliderController().Get()->GetTimeSliderArgs().OnScrubPositionChanged = FOnScrubPositionChanged::CreateRaw(this, &SVisualLogger::OnScrubPositionChanged);
+	FLogVisualizer::Get().GetTimeSliderController().Get()->GetTimeSliderArgs().OnScrubPositionChanged = FVisualLoggerTimeSliderArgs::FOnScrubPositionChanged::CreateRaw(this, &SVisualLogger::OnScrubPositionChanged);
 }
 
 void SVisualLogger::OnNewWorld(UWorld* NewWorld)
diff --git a/Engine/Source/Developer/LogVisualizer/Private/SVisualLoggerTimeSlider.cpp b/Engine/Source/Developer/LogVisualizer/Private/SVisualLoggerTimeSlider.cpp
--- a/Engine/Source/Developer/LogVisualizer/Private/SVisualLoggerTimeSlider.cpp
+++ b/Engine/Source/Developer/LogVisualizer/Private/SVisualLoggerTimeSlider.cpp
@@ -1,11 +1,12 @@
 // Copyright 1998-2018 Epic Games, Inc. All Rights Reserved.
 
 #include "SVisualLoggerTimeSlider.h"
+#include "VisualLoggerTimeSliderController.h"
 
 #define LOCTEXT_NAMESPACE "STimeSlider"
 
 
-void SVisualLoggerTimeSlider::Construct( const SVisualLoggerTimeSlider::FArguments& InArgs, TSharedRef<ITimeSliderController> InTimeSliderController )
+void SVisualLoggerTimeSlider::Construct( const SVisualLoggerTimeSlider::FArguments& InArgs, TSharedRef<FVisualLoggerTimeSliderController> InTimeSliderController )
 {
 	TimeSliderController = InTimeSliderController;
 	bMirrorLabels = InArgs._MirrorLabels;
diff --git a/Engine/Source/Developer/LogVisualizer/Private/SVisualLoggerTimeSlider.h b/Engine/Source/Developer/LogVisualizer/Private/SVisualLoggerTimeSlider.h
--- a/Engine/Source/Developer/LogVisualizer/Private/SVisualLoggerTimeSlider.h
+++ b/Engine/Source/Developer/LogVisualizer/Private/SVisualLoggerTimeSlider.h
@@ -4,13 +4,14 @@
 
 #include "CoreMinimal.h"
 #include "Input/Reply.h"
+#include "Widgets/SCompoundWidget.h"
 #include "Widgets/DeclarativeSyntaxSupport.h"
-#include "ITimeSlider.h"
 
 class FPaintArgs;
 class FSlateWindowElementList;
+class FVisualLoggerTimeSliderController;
 
-class SVisualLoggerTimeSlider : public ITimeSlider
+class SVisualLoggerTimeSlider : public SCompoundWidget
 {
 public:
 
@@ -27,7 +28,7 @@ public:
 	 * 
 	 * @param InArgs   A declaration from which to construct the widget
 	 */
-	void Construct( const FArguments& InArgs, TSharedRef<ITimeSliderController> InTimeSliderController );
+	void Construct( const FArguments& InArgs, TSharedRef<FVisualLoggerTimeSliderController> InTimeSliderController );
 
 protected:
 	// SWidget interface
@@ -38,7 +39,7 @@ protected:
 	virtual FReply OnMouseMove( const FGeometry& MyGeometry, const FPointerEvent& MouseEvent ) override;
 	virtual FReply OnMouseWheel( const FGeometry& MyGeometry, const FPointerEvent& MouseEvent ) override;
 private:
-	TSharedPtr<ITimeSliderController> TimeSliderController;
+	TSharedPtr<FVisualLoggerTimeSliderController> TimeSliderController;
 	bool bMirrorLabels;
 };
 
diff --git a/Engine/Source/Developer/LogVisualizer/Private/SVisualLoggerView.cpp b/Engine/Source/Developer/LogVisualizer/Private/SVisualLoggerView.cpp
--- a/Engine/Source/Developer/LogVisualizer/Private/SVisualLoggerView.cpp
+++ b/Engine/Source/Developer/LogVisualizer/Private/SVisualLoggerView.cpp
@@ -11,7 +11,6 @@
 #include "SVisualLoggerSectionOverlay.h"
 #include "SVisualLoggerTimeline.h"
 #include "SVisualLoggerTimelinesContainer.h"
-#include "ITimeSlider.h"
 #include "SVisualLoggerTimeSlider.h"
 #include "VisualLoggerTimeSliderController.h"
 #include "Widgets/Input/SSearchBox.h"
@@ -86,8 +85,8 @@ void SVisualLoggerView::Construct(const FArguments& InArgs, const TSharedRef<FUI
 
 	// Create the top and bottom sliders
 	const bool bMirrorLabels = true;
-	TSharedRef<ITimeSlider> TopTimeSlider = SNew(SVisualLoggerTimeSlider, FLogVisualizer::Get().GetTimeSliderController().ToSharedRef()).MirrorLabels(bMirrorLabels);
-	TSharedRef<ITimeSlider> BottomTimeSlider = SNew(SVisualLoggerTimeSlider, FLogVisualizer::Get().GetTimeSliderController().ToSharedRef()).MirrorLabels(bMirrorLabels);
+	TSharedRef<SVisualLoggerTimeSlider> TopTimeSlider = SNew(SVisualLoggerTimeSlider, FLogVisualizer::Get().GetTimeSliderController().ToSharedRef()).MirrorLabels(bMirrorLabels);
+	TSharedRef<SVisualLoggerTimeSlider> BottomTimeSlider = SNew(SVisualLoggerTimeSlider, FLogVisualizer::Get().GetTimeSliderController().ToSharedRef()).MirrorLabels(bMirrorLabels);
 
 	TSharedRef<SScrollBar> ScrollBar =
 		SNew(SScrollBar)
diff --git a/Engine/Source/Developer/LogVisualizer/Private/VisualLoggerTimeSliderController.cpp b/Engine/Source/Developer/LogVisualizer/Private/VisualLoggerTimeSliderController.cpp
--- a/Engine/Source/Developer/LogVisualizer/Private/VisualLoggerTimeSliderController.cpp
+++ b/Engine/Source/Developer/LogVisualizer/Private/VisualLoggerTimeSliderController.cpp
@@ -414,7 +414,7 @@ FReply FVisualLoggerTimeSliderController::OnMouseMove( SWidget& WidgetOwner, con
 					NewViewOutputMax = LocalClampMax;
 				}
 
-				TimeSliderArgs.OnViewRangeChanged.ExecuteIfBound(TRange<float>(NewViewOutputMin, NewViewOutputMax), EViewRangeInterpolation::Immediate);
+				TimeSliderArgs.OnViewRangeChanged.ExecuteIfBound(TRange<float>(NewViewOutputMin, NewViewOutputMax));
 				if (Scrollbar.IsValid())
 				{
 					float InOffsetFraction = (NewViewOutputMin - LocalClampMin) / (LocalClampMax - LocalClampMin);
@@ -592,7 +592,7 @@ FReply FVisualLoggerTimeSliderController::OnMouseWheel( SWidget& WidgetOwner, co
 					NewViewOutputMax = LocalClampMax;
 				}
 
-				TimeSliderArgs.OnViewRangeChanged.ExecuteIfBound(TRange<float>(NewViewOutputMin, NewViewOutputMax), EViewRangeInterpolation::Immediate);
+				TimeSliderArgs.OnViewRangeChanged.ExecuteIfBound(TRange<float>(NewViewOutputMin, NewViewOutputMax));
 				if (Scrollbar.IsValid())
 				{
 					float InOffsetFraction = (NewViewOutputMin - LocalClampMin) / (LocalClampMax - LocalClampMin);
diff --git a/Engine/Source/Developer/LogVisualizer/Private/VisualLoggerTimeSliderController.h b/Engine/Source/Developer/LogVisualizer/Private/VisualLoggerTimeSliderController.h
--- a/Engine/Source/Developer/LogVisualizer/Private/VisualLoggerTimeSliderController.h
+++ b/Engine/Source/Developer/LogVisualizer/Private/VisualLoggerTimeSliderController.h
@@ -8,17 +8,48 @@
 #include "Input/Reply.h"
 #include "Widgets/SWidget.h"
 #include "Widgets/Layout/SScrollBar.h"
-#include "ITimeSlider.h"
 
 class FSlateWindowElementList;
 
-struct FVisualLoggerTimeSliderArgs : FTimeSliderArgs
+struct FVisualLoggerTimeSliderArgs
 {
+	DECLARE_DELEGATE_OneParam( FOnRangeChanged, TRange<float> )
+	DECLARE_DELEGATE_TwoParams( FOnScrubPositionChanged, float, bool )
+
 	FVisualLoggerTimeSliderArgs()
-		: CursorSize(0.05f)
-	{
-		ViewRange = TRange<float>(0.0f, 5.0f);
-	}
+		: ScrubPosition(0)
+		, ViewRange( TRange<float>(0.0f, 5.0f) )
+		, ClampRange( TRange<float>(-FLT_MAX/2.f, FLT_MAX/2.f) )
+		, AllowZoom(true)
+		, CursorSize(0.05f)
+	{}
+
+	/** The scrub position */
+	TAttribute<float> ScrubPosition;
+
+	/** Called when the scrub position changes */
+	FOnScrubPositionChanged OnScrubPositionChanged;
+
+	/** Called right before the scrubber begins to move */
+	FSimpleDelegate OnBeginScrubberMovement;
+
+	/** Called right after the scrubber handle is released by the user */
+	FSimpleDelegate OnEndScrubberMovement;
+
+	/** View time range */
+	TAttribute< TRange<float> > ViewRange;
+
+	/** Clamp time range */
+	TAttribute< TRange<float> > ClampRange;
+
+	/** Called when the view range changes */
+	FOnRangeChanged OnViewRangeChanged;
+
+	/** Attribute defining the time snap interval */
+	TAttribute<float> TimeSnapInterval;
+
+	/** If we are allowed to zoom */
+	bool AllowZoom;
 
 	/** Cursor range for data like histogram graphs, etc. */
 	TAttribute< float > CursorSize;
@@ -28,7 +59,7 @@ struct FVisualLoggerTimeSliderArgs : FTimeSliderArgs
  * A time slider controller for sequencer
  * Draws and manages time data for a Sequencer
  */
-class FVisualLoggerTimeSliderController : public ITimeSliderController
+class FVisualLoggerTimeSliderController
 {
 public:
 	FVisualLoggerTimeSliderController(const FVisualLoggerTimeSliderArgs& InArgs);
@@ -45,13 +76,12 @@ public:
 	float DetermineOptimalSpacing(float InPixelsPerInput, uint32 MinTick, float MinTickSpacing) const;
 	void SetTimesliderArgs(const FVisualLoggerTimeSliderArgs& InArgs);
 
-	/** ITimeSliderController Interface */
-	virtual int32 OnPaintTimeSlider( bool bMirrorLabels, const FGeometry& AllottedGeometry, const FSlateRect& MyCullingRect, FSlateWindowElementList& OutDrawElements, int32 LayerId, const FWidgetStyle& InWidgetStyle, bool bParentEnabled ) const override;
-	virtual FReply OnMouseButtonDown( SWidget& WidgetOwner, const FGeometry& MyGeometry, const FPointerEvent& MouseEvent ) override;
-	virtual FReply OnMouseButtonUp( SWidget& WidgetOwner, const FGeometry& MyGeometry, const FPointerEvent& MouseEvent ) override;
-	virtual FReply OnMouseMove( SWidget& WidgetOwner, const FGeometry& MyGeometry, const FPointerEvent& MouseEvent ) override;
-	virtual FReply OnMouseWheel( SWidget& WidgetOwner, const FGeometry& MyGeometry, const FPointerEvent& MouseEvent ) override;
-	virtual FCursorReply OnCursorQuery( TSharedRef<const SWidget> WidgetOwner, const FGeometry& MyGeometry, const FPointerEvent& CursorEvent) const override { return FCursorReply::Unhandled(); }
+	int32 OnPaintTimeSlider( bool bMirrorLabels, const FGeometry& AllottedGeometry, const FSlateRect& MyCullingRect, FSlateWindowElementList& OutDrawElements, int32 LayerId, const FWidgetStyle& InWidgetStyle, bool bParentEnabled ) const;
+	FReply OnMouseButtonDown( SWidget& WidgetOwner, const FGeometry& MyGeometry, const FPointerEvent& MouseEvent );
+	FReply OnMouseButtonUp( SWidget& WidgetOwner, const FGeometry& MyGeometry, const FPointerEvent& MouseEvent );
+	FReply OnMouseMove( SWidget& WidgetOwner, const FGeometry& MyGeometry, const FPointerEvent& MouseEvent );
+	FReply OnMouseWheel( SWidget& WidgetOwner, const FGeometry& MyGeometry, const FPointerEvent& MouseEvent );
+	FCursorReply OnCursorQuery( TSharedRef<const SWidget> WidgetOwner, const FGeometry& MyGeometry, const FPointerEvent& CursorEvent) const { return FCursorReply::Unhandled(); }
 
 	/**
 	 * Draws major tick lines in the section view                                                              
diff --git a/Engine/Source/Developer/Profiler/Private/Widgets/SDataGraph.cpp b/Engine/Source/Developer/Profiler/Private/Widgets/SDataGraph.cpp
--- a/Engine/Source/Developer/Profiler/Private/Widgets/SDataGraph.cpp
+++ b/Engine/Source/Developer/Profiler/Private/Widgets/SDataGraph.cpp
@@ -1137,10 +1137,8 @@ void SDataGraph::OnMouseLeave( const FPointerEvent& MouseEvent )
 
 FReply SDataGraph::OnMouseWheel( const FGeometry& MyGeometry, const FPointerEvent& MouseEvent )
 {
-	// 2^ 3=8
-	// 2^11=2048
 	MouseWheelAcc += MouseEvent.GetWheelDelta() * 0.25f;
-	MouseWheelAcc = FMath::Clamp( MouseWheelAcc, 3.0f, 11.0f );
+	MouseWheelAcc = FMath::Clamp( MouseWheelAcc, 3.0f, 13.0f ); // 2^3=8, 2^13=8192
 
 	ScaleY = FMath::Pow( 2.0f, MouseWheelAcc );
 
diff --git a/Engine/Source/Developer/XGEController/Private/XGEControllerModule.cpp b/Engine/Source/Developer/XGEController/Private/XGEControllerModule.cpp
--- a/Engine/Source/Developer/XGEController/Private/XGEControllerModule.cpp
+++ b/Engine/Source/Developer/XGEController/Private/XGEControllerModule.cpp
@@ -177,6 +177,11 @@ bool FXGEControllerModule::IsSupported()
 
 #if PLATFORM_WINDOWS
 
+	if (!FPlatformProcess::SupportsMultithreading())
+	{
+		return false; // current implementation requires worker threads
+	}
+
 	// Check the command line to see if the XGE controller has been enabled/disabled.
 	// This overrides the value of the console variable.
 	if (FParse::Param(FCommandLine::Get(), TEXT("xgecontroller")))
diff --git a/Engine/Source/Editor/CurveAssetEditor/CurveAssetEditor.Build.cs b/Engine/Source/Editor/CurveAssetEditor/CurveAssetEditor.Build.cs
--- a/Engine/Source/Editor/CurveAssetEditor/CurveAssetEditor.Build.cs
+++ b/Engine/Source/Editor/CurveAssetEditor/CurveAssetEditor.Build.cs
@@ -13,6 +13,7 @@ public class CurveAssetEditor : ModuleRules
 			new string[] {
 				"Core",
 				"CoreUObject",
+				"CurveEditor",
 				"Engine",
 				"Slate",
 				"SlateCore",
diff --git a/Engine/Source/Editor/CurveAssetEditor/Private/CurveAssetEditor.cpp b/Engine/Source/Editor/CurveAssetEditor/Private/CurveAssetEditor.cpp
--- a/Engine/Source/Editor/CurveAssetEditor/Private/CurveAssetEditor.cpp
+++ b/Engine/Source/Editor/CurveAssetEditor/Private/CurveAssetEditor.cpp
@@ -10,7 +10,7 @@
 #include "CurveAssetEditorModule.h"
 
 #include "SCurveEditor.h"
-#include "RichCurveEditorCommands.h"
+#include "CurveEditorCommands.h"
 //#include "Toolkits/IToolkitHost.h"
 #include "Widgets/Docking/SDockTab.h"
 #include "Widgets/Input/SNumericDropDown.h"
@@ -185,19 +185,19 @@ TSharedPtr<FExtender> FCurveAssetEditor::GetToolbarExtender()
 		{
 			ToolbarBuilder.BeginSection("Curve");
 			{
-				ToolbarBuilder.AddToolBarButton(FRichCurveEditorCommands::Get().ZoomToFitHorizontal);
-				ToolbarBuilder.AddToolBarButton(FRichCurveEditorCommands::Get().ZoomToFitVertical);
-				ToolbarBuilder.AddToolBarButton(FRichCurveEditorCommands::Get().ZoomToFit);
+				ToolbarBuilder.AddToolBarButton(FCurveEditorCommands::Get().ZoomToFitHorizontal);
+				ToolbarBuilder.AddToolBarButton(FCurveEditorCommands::Get().ZoomToFitVertical);
+				ToolbarBuilder.AddToolBarButton(FCurveEditorCommands::Get().ZoomToFit);
 			}
 			ToolbarBuilder.EndSection();
 
 			ToolbarBuilder.BeginSection("Interpolation");
 			{
-				ToolbarBuilder.AddToolBarButton(FRichCurveEditorCommands::Get().InterpolationCubicAuto);
-				ToolbarBuilder.AddToolBarButton(FRichCurveEditorCommands::Get().InterpolationCubicUser);
-				ToolbarBuilder.AddToolBarButton(FRichCurveEditorCommands::Get().InterpolationCubicBreak);
-				ToolbarBuilder.AddToolBarButton(FRichCurveEditorCommands::Get().InterpolationLinear);
-				ToolbarBuilder.AddToolBarButton(FRichCurveEditorCommands::Get().InterpolationConstant);
+				ToolbarBuilder.AddToolBarButton(FCurveEditorCommands::Get().InterpolationCubicAuto);
+				ToolbarBuilder.AddToolBarButton(FCurveEditorCommands::Get().InterpolationCubicUser);
+				ToolbarBuilder.AddToolBarButton(FCurveEditorCommands::Get().InterpolationCubicBreak);
+				ToolbarBuilder.AddToolBarButton(FCurveEditorCommands::Get().InterpolationLinear);
+				ToolbarBuilder.AddToolBarButton(FCurveEditorCommands::Get().InterpolationConstant);
 			}
 			ToolbarBuilder.EndSection();
 
@@ -211,9 +211,9 @@ TSharedPtr<FExtender> FCurveAssetEditor::GetToolbarExtender()
 
 			ToolbarBuilder.BeginSection("Snap");
 			{
-				ToolbarBuilder.AddToolBarButton(FRichCurveEditorCommands::Get().ToggleInputSnapping);
+				ToolbarBuilder.AddToolBarButton(FCurveEditorCommands::Get().ToggleInputSnapping);
 				ToolbarBuilder.AddWidget(InputSnapWidget);
-				ToolbarBuilder.AddToolBarButton(FRichCurveEditorCommands::Get().ToggleOutputSnapping);
+				ToolbarBuilder.AddToolBarButton(FCurveEditorCommands::Get().ToggleOutputSnapping);
 				ToolbarBuilder.AddWidget(OutputSnapWidget);
 			}
 			ToolbarBuilder.EndSection();
@@ -271,11 +271,11 @@ TSharedRef<SWidget> FCurveAssetEditor::MakeCurveEditorCurveOptionsMenu()
 		{
 			MenuBuilder.BeginSection( "Pre-Infinity Extrapolation", LOCTEXT( "CurveEditorMenuPreInfinityExtrapHeader", "Extrapolation" ) );
 			{
-				MenuBuilder.AddMenuEntry( FRichCurveEditorCommands::Get().SetPreInfinityExtrapCycle);
-				MenuBuilder.AddMenuEntry( FRichCurveEditorCommands::Get().SetPreInfinityExtrapCycleWithOffset);
-				MenuBuilder.AddMenuEntry( FRichCurveEditorCommands::Get().SetPreInfinityExtrapOscillate);
-				MenuBuilder.AddMenuEntry( FRichCurveEditorCommands::Get().SetPreInfinityExtrapLinear);
-				MenuBuilder.AddMenuEntry( FRichCurveEditorCommands::Get().SetPreInfinityExtrapConstant);
+				MenuBuilder.AddMenuEntry( FCurveEditorCommands::Get().SetPreInfinityExtrapCycle);
+				MenuBuilder.AddMenuEntry( FCurveEditorCommands::Get().SetPreInfinityExtrapCycleWithOffset);
+				MenuBuilder.AddMenuEntry( FCurveEditorCommands::Get().SetPreInfinityExtrapOscillate);
+				MenuBuilder.AddMenuEntry( FCurveEditorCommands::Get().SetPreInfinityExtrapLinear);
+				MenuBuilder.AddMenuEntry( FCurveEditorCommands::Get().SetPreInfinityExtrapConstant);
 			}
 			MenuBuilder.EndSection();
 		}
@@ -284,11 +284,11 @@ TSharedRef<SWidget> FCurveAssetEditor::MakeCurveEditorCurveOptionsMenu()
 		{
 			MenuBuilder.BeginSection( "Post-Infinity Extrapolation", LOCTEXT( "CurveEditorMenuPostInfinityExtrapHeader", "Extrapolation" ) );
 			{
-				MenuBuilder.AddMenuEntry( FRichCurveEditorCommands::Get().SetPostInfinityExtrapCycle);
-				MenuBuilder.AddMenuEntry( FRichCurveEditorCommands::Get().SetPostInfinityExtrapCycleWithOffset);
-				MenuBuilder.AddMenuEntry( FRichCurveEditorCommands::Get().SetPostInfinityExtrapOscillate);
-				MenuBuilder.AddMenuEntry( FRichCurveEditorCommands::Get().SetPostInfinityExtrapLinear);
-				MenuBuilder.AddMenuEntry( FRichCurveEditorCommands::Get().SetPostInfinityExtrapConstant);
+				MenuBuilder.AddMenuEntry( FCurveEditorCommands::Get().SetPostInfinityExtrapCycle);
+				MenuBuilder.AddMenuEntry( FCurveEditorCommands::Get().SetPostInfinityExtrapCycleWithOffset);
+				MenuBuilder.AddMenuEntry( FCurveEditorCommands::Get().SetPostInfinityExtrapOscillate);
+				MenuBuilder.AddMenuEntry( FCurveEditorCommands::Get().SetPostInfinityExtrapLinear);
+				MenuBuilder.AddMenuEntry( FCurveEditorCommands::Get().SetPostInfinityExtrapConstant);
 			}
 			MenuBuilder.EndSection();
 		}
@@ -296,8 +296,8 @@ TSharedRef<SWidget> FCurveAssetEditor::MakeCurveEditorCurveOptionsMenu()
 
 	FMenuBuilder MenuBuilder( true, TrackWidget->GetCommands());
 
-	MenuBuilder.AddMenuEntry( FRichCurveEditorCommands::Get().BakeCurve);
-	MenuBuilder.AddMenuEntry( FRichCurveEditorCommands::Get().ReduceCurve);
+	MenuBuilder.AddMenuEntry( FCurveEditorCommands::Get().BakeCurve);
+	MenuBuilder.AddMenuEntry( FCurveEditorCommands::Get().ReduceCurve);
 
 	MenuBuilder.AddSubMenu(
 		LOCTEXT( "PreInfinitySubMenu", "Pre-Infinity" ),
diff --git a/Engine/Source/Editor/CurveAssetEditor/Private/CurveAssetEditorModule.cpp b/Engine/Source/Editor/CurveAssetEditor/Private/CurveAssetEditorModule.cpp
--- a/Engine/Source/Editor/CurveAssetEditor/Private/CurveAssetEditorModule.cpp
+++ b/Engine/Source/Editor/CurveAssetEditor/Private/CurveAssetEditorModule.cpp
@@ -3,7 +3,7 @@
 #include "CurveAssetEditorModule.h"
 #include "Modules/ModuleManager.h"
 #include "CurveAssetEditor.h"
-#include "RichCurveEditorCommands.h"
+#include "CurveEditorCommands.h"
 //#include "Toolkits/ToolkitManager.h"
 
 IMPLEMENT_MODULE( FCurveAssetEditorModule, CurveAssetEditor );
@@ -14,7 +14,7 @@ const FName FCurveAssetEditorModule::CurveAssetEditorAppIdentifier( TEXT( "Curve
 void FCurveAssetEditorModule::StartupModule()
 {
 	MenuExtensibilityManager = MakeShareable(new FExtensibilityManager);
-	FRichCurveEditorCommands::Register();
+	FCurveEditorCommands::Register();
 }
 
 void FCurveAssetEditorModule::ShutdownModule()
diff --git a/Engine/Source/Editor/CurveEditor/Private/CurveEditor.cpp b/Engine/Source/Editor/CurveEditor/Private/CurveEditor.cpp
new file mode 100644
--- /dev/null
+++ b/Engine/Source/Editor/CurveEditor/Private/CurveEditor.cpp
@@ -0,0 +1,837 @@
+// Copyright 1998-2018 Epic Games, Inc. All Rights Reserved.
+
+#include "CurveEditor.h"
+#include "Layout/Geometry.h"
+#include "CurveEditorScreenSpace.h"
+#include "CurveEditorSnapMetrics.h"
+#include "CurveEditorCommands.h"
+#include "CurveEditorSettings.h"
+#include "CurveDrawInfo.h"
+#include "Framework/Commands/GenericCommands.h"
+#include "Framework/Commands/UICommandList.h"
+#include "Editor.h"
+#include "ScopedTransaction.h"
+
+#define LOCTEXT_NAMESPACE "CurveEditor"
+
+FCurveModelID FCurveModelID::Unique()
+{
+	static uint32 CurrentID = 1;
+
+	FCurveModelID ID;
+	ID.ID = CurrentID++;
+	return ID;
+}
+
+FCurveEditor::FCurveEditor()
+	: Bounds(new FStaticCurveEditorBounds)
+	, CachedPhysicalSize(0.f, 0.f)
+{
+	Settings = GetMutableDefault<UCurveEditorSettings>();
+}
+
+FCurveModel* FCurveEditor::FindCurve(FCurveModelID CurveID) const
+{
+	const TUniquePtr<FCurveModel>* Ptr = CurveData.Find(CurveID);
+	return Ptr ? Ptr->Get() : nullptr;
+}
+
+const TMap<FCurveModelID, TUniquePtr<FCurveModel>>& FCurveEditor::GetCurves() const
+{
+	return CurveData;
+}
+
+FCurveModelID FCurveEditor::AddCurve(TUniquePtr<FCurveModel>&& InCurve)
+{
+	FCurveModelID NewID = FCurveModelID::Unique();
+	CurveData.Add(NewID, MoveTemp(InCurve));
+	return NewID;
+}
+
+void FCurveEditor::RemoveCurve(FCurveModelID InCurveID)
+{
+	CurveData.Remove(InCurveID);
+	Selection.Remove(InCurveID);
+}
+
+void FCurveEditor::SetBounds(TUniquePtr<ICurveEditorBounds>&& InBounds)
+{
+	check(InBounds.IsValid());
+	Bounds = MoveTemp(InBounds);
+}
+
+bool FCurveEditor::ShouldAutoFrame() const
+{
+	return Settings->GetAutoFrameCurveEditor();
+}
+
+void FCurveEditor::BindCommands()
+{
+	UCurveEditorSettings* CurveSettings = Settings;
+
+	CommandList = MakeShared<FUICommandList>();
+
+	CommandList->MapAction(FGenericCommands::Get().Undo,   FExecuteAction::CreateLambda([]{ GEditor->UndoTransaction(); }));
+	CommandList->MapAction(FGenericCommands::Get().Redo,   FExecuteAction::CreateLambda([]{ GEditor->RedoTransaction(); }));
+	CommandList->MapAction(FGenericCommands::Get().Delete, FExecuteAction::CreateSP(this, &FCurveEditor::DeleteSelection));
+
+	CommandList->MapAction(FCurveEditorCommands::Get().ZoomToFitAll,        FExecuteAction::CreateSP(this, &FCurveEditor::ZoomToFit, EAxisList::All));
+	CommandList->MapAction(FCurveEditorCommands::Get().ZoomToFit,           FExecuteAction::CreateSP(this, &FCurveEditor::ZoomToFitSelection, EAxisList::All));
+	CommandList->MapAction(FCurveEditorCommands::Get().ZoomToFitHorizontal, FExecuteAction::CreateSP(this, &FCurveEditor::ZoomToFitSelection, EAxisList::X));
+	CommandList->MapAction(FCurveEditorCommands::Get().ZoomToFitVertical,   FExecuteAction::CreateSP(this, &FCurveEditor::ZoomToFitSelection, EAxisList::Y));
+
+	{
+		FExecuteAction   ToggleInputSnapping     = FExecuteAction::CreateSP(this,   &FCurveEditor::ToggleInputSnapping);
+		FIsActionChecked IsInputSnappingEnabled  = FIsActionChecked::CreateSP(this, &FCurveEditor::IsInputSnappingEnabled);
+		FExecuteAction   ToggleOutputSnapping    = FExecuteAction::CreateSP(this,   &FCurveEditor::ToggleOutputSnapping);
+		FIsActionChecked IsOutputSnappingEnabled = FIsActionChecked::CreateSP(this, &FCurveEditor::IsOutputSnappingEnabled);
+
+		CommandList->MapAction(FCurveEditorCommands::Get().ToggleInputSnapping, ToggleInputSnapping, FCanExecuteAction(), IsInputSnappingEnabled);
+		CommandList->MapAction(FCurveEditorCommands::Get().ToggleOutputSnapping, ToggleOutputSnapping, FCanExecuteAction(), IsOutputSnappingEnabled);
+	}
+
+	CommandList->MapAction(FCurveEditorCommands::Get().FlattenTangents,    FExecuteAction::CreateSP(this, &FCurveEditor::FlattenSelection));
+	CommandList->MapAction(FCurveEditorCommands::Get().StraightenTangents, FExecuteAction::CreateSP(this, &FCurveEditor::StraightenSelection));
+	CommandList->MapAction(FCurveEditorCommands::Get().BakeCurve,          FExecuteAction::CreateSP(this, &FCurveEditor::BakeSelection));
+	CommandList->MapAction(FCurveEditorCommands::Get().ReduceCurve,        FExecuteAction::CreateSP(this, &FCurveEditor::SimplifySelection, 0.1f));
+
+	// Tangent Visibility
+	{
+		FExecuteAction SetAllTangents          = FExecuteAction::CreateUObject(Settings, &UCurveEditorSettings::SetTangentVisibility, ECurveEditorTangentVisibility::AllTangents);
+		FExecuteAction SetSelectedKeyTangents  = FExecuteAction::CreateUObject(Settings, &UCurveEditorSettings::SetTangentVisibility, ECurveEditorTangentVisibility::SelectedKeys);
+		FExecuteAction SetNoTangents           = FExecuteAction::CreateUObject(Settings, &UCurveEditorSettings::SetTangentVisibility, ECurveEditorTangentVisibility::NoTangents);
+
+		FIsActionChecked IsAllTangents         = FIsActionChecked::CreateLambda( [CurveSettings]{ return CurveSettings->GetTangentVisibility() == ECurveEditorTangentVisibility::AllTangents; } );
+		FIsActionChecked IsSelectedKeyTangents = FIsActionChecked::CreateLambda( [CurveSettings]{ return CurveSettings->GetTangentVisibility() == ECurveEditorTangentVisibility::SelectedKeys; } );
+		FIsActionChecked IsNoTangents          = FIsActionChecked::CreateLambda( [CurveSettings]{ return CurveSettings->GetTangentVisibility() == ECurveEditorTangentVisibility::NoTangents; } );
+
+		CommandList->MapAction(FCurveEditorCommands::Get().SetAllTangentsVisibility, SetAllTangents, FCanExecuteAction(), IsAllTangents);
+		CommandList->MapAction(FCurveEditorCommands::Get().SetSelectedKeysTangentVisibility, SetSelectedKeyTangents, FCanExecuteAction(), IsSelectedKeyTangents);
+		CommandList->MapAction(FCurveEditorCommands::Get().SetNoTangentsVisibility, SetNoTangents, FCanExecuteAction(), IsNoTangents);
+	}
+
+	CommandList->MapAction(FCurveEditorCommands::Get().ToggleAutoFrameCurveEditor,
+		FExecuteAction::CreateLambda( [CurveSettings]{ CurveSettings->SetAutoFrameCurveEditor( !CurveSettings->GetAutoFrameCurveEditor() ); } ),
+		FCanExecuteAction(),
+		FIsActionChecked::CreateLambda( [CurveSettings]{ return CurveSettings->GetAutoFrameCurveEditor(); } )
+	);
+
+	CommandList->MapAction(FCurveEditorCommands::Get().ToggleShowCurveEditorCurveToolTips,
+		FExecuteAction::CreateLambda( [CurveSettings]{ CurveSettings->SetShowCurveEditorCurveToolTips( !CurveSettings->GetShowCurveEditorCurveToolTips() ); } ),
+		FCanExecuteAction(),
+		FIsActionChecked::CreateLambda( [CurveSettings]{ return CurveSettings->GetShowCurveEditorCurveToolTips(); } ) );
+}
+
+FCurveEditorScreenSpace FCurveEditor::GetScreenSpace() const
+{
+	// Retrieve the current bounds and cache them
+	double InputMin = 0, InputMax = 1;
+	Bounds->GetInputBounds(InputMin, InputMax);
+
+	double OutputMin = 0, OutputMax = 1;
+	Bounds->GetOutputBounds(OutputMin, OutputMax);
+
+	return FCurveEditorScreenSpace(CachedPhysicalSize, InputMin, InputMax, OutputMin, OutputMax);
+}
+
+FCurveEditorSnapMetrics FCurveEditor::GetSnapMetrics() const
+{
+	FCurveEditorSnapMetrics Metrics;
+
+	Metrics.bSnapOutputValues  = OutputSnapEnabledAttribute.Get();
+	Metrics.OutputSnapInterval = OutputSnapIntervalAttribute.Get();
+	Metrics.bSnapInputValues   = InputSnapEnabledAttribute.Get();
+	Metrics.InputSnapRate      = InputSnapRateAttribute.Get();
+
+	return Metrics;
+}
+
+void FCurveEditor::UpdateGeometry(const FGeometry& NewGeometry)
+{
+	if (CachedPhysicalSize.X != 0 && CachedPhysicalSize.Y != 0)
+	{
+		// Retrieve the current bounds and cache them
+		double InputMin = 0, InputMax = 1;
+		Bounds->GetInputBounds(InputMin, InputMax);
+
+		double OutputMin = 0, OutputMax = 1;
+		Bounds->GetOutputBounds(OutputMin, OutputMax);
+
+		// Increase the visible input/output ranges based on the new size of the panel
+		const float OldWidth = CachedPhysicalSize.X;
+		if (NewGeometry.GetLocalSize().X != OldWidth)
+		{
+			const double PixelToInputRatio = (InputMax - InputMin) / OldWidth;
+			InputMax += PixelToInputRatio * (NewGeometry.GetLocalSize().X - OldWidth);
+
+			Bounds->SetInputBounds(InputMin, InputMax);
+		}
+
+		const float OldHeight = CachedPhysicalSize.Y;
+		if (NewGeometry.GetLocalSize().Y != OldHeight)
+		{
+			const double PixelToOutputRatio = (OutputMax - OutputMin) / OldHeight;
+			OutputMin -= PixelToOutputRatio * (NewGeometry.GetLocalSize().Y - OldHeight);
+
+			Bounds->SetOutputBounds(OutputMin, OutputMax);
+		}
+	}
+
+	CachedPhysicalSize = NewGeometry.GetLocalSize();
+}
+
+void FCurveEditor::Zoom(float Amount)
+{
+	FCurveEditorScreenSpace ScreenSpace = GetScreenSpace();
+
+	const double InputOrigin  = (ScreenSpace.GetInputMax()  - ScreenSpace.GetInputMin())  * 0.5;
+	const double OutputOrigin = (ScreenSpace.GetOutputMax() - ScreenSpace.GetOutputMin()) * 0.5;
+
+	ZoomAround(Amount, InputOrigin, OutputOrigin);
+}
+
+void FCurveEditor::ZoomAround(float Amount, double InputOrigin, double OutputOrigin)
+{
+	FCurveEditorScreenSpace ScreenSpace = GetScreenSpace();
+
+	double InputMin  = ScreenSpace.GetInputMin();
+	double InputMax  = ScreenSpace.GetInputMax();
+	double OutputMin = ScreenSpace.GetOutputMin();
+	double OutputMax = ScreenSpace.GetOutputMax();
+
+	InputMin = InputOrigin - (InputOrigin - InputMin) * Amount;
+	InputMax = InputOrigin + (InputMax - InputOrigin) * Amount;
+
+	OutputMin = OutputOrigin - (OutputOrigin - OutputMin) * Amount;
+	OutputMax = OutputOrigin + (OutputMax - OutputOrigin) * Amount;
+
+	Bounds->SetInputBounds(InputMin, InputMax);
+	Bounds->SetOutputBounds(OutputMin, OutputMax);
+}
+
+void FCurveEditor::ZoomToFit(EAxisList::Type Axes)
+{
+	TArray<FCurveModelID> CurveModelIDs;
+	for (const TTuple<FCurveModelID, TUniquePtr<FCurveModel>>& Pair : CurveData)
+	{
+		CurveModelIDs.Add(Pair.Key);
+	}
+
+	ZoomToFitCurves(MakeArrayView(CurveModelIDs.GetData(), CurveModelIDs.Num()), Axes);
+}
+
+void FCurveEditor::ZoomToFitCurves(TArrayView<const FCurveModelID> CurveModelIDs, EAxisList::Type Axes)
+{
+	if (CurveData.Num() && CurveModelIDs.Num())
+	{
+		TArrayView<const FKeyHandle> Empty;
+
+		double InputMin = TNumericLimits<double>::Max(), InputMax = TNumericLimits<double>::Lowest(), OutputMin = TNumericLimits<double>::Max(), OutputMax = TNumericLimits<double>::Lowest();
+
+		TArray<FKeyHandle> KeyHandlesScratch;
+		TArray<FKeyPosition> KeyPositionsScratch;
+
+		for (const FCurveModelID CurveModelID : CurveModelIDs)
+		{
+			if (const FCurveModel* Curve = FindCurve(CurveModelID))
+			{
+				KeyHandlesScratch.Reset();
+				Curve->GetKeys(*this,
+					TNumericLimits<double>::Lowest(),
+					TNumericLimits<double>::Max(),
+					TNumericLimits<double>::Lowest(),
+					TNumericLimits<double>::Max(),
+					KeyHandlesScratch);
+
+				if (KeyHandlesScratch.Num())
+				{
+					KeyPositionsScratch.SetNum(KeyHandlesScratch.Num());
+					Curve->GetKeyPositions(KeyHandlesScratch, KeyPositionsScratch);
+
+					for (FKeyPosition Key : KeyPositionsScratch)
+					{
+						InputMin  = FMath::Min(InputMin, Key.InputValue);
+						InputMax  = FMath::Max(InputMax, Key.InputValue);
+						OutputMin = FMath::Min(OutputMin, Key.OutputValue);
+						OutputMax = FMath::Max(OutputMax, Key.OutputValue);
+					}
+				}
+			}
+		}
+
+		if (InputMin != TNumericLimits<double>::Max() && InputMax != TNumericLimits<double>::Lowest() && 
+			OutputMin != TNumericLimits<double>::Max() && OutputMax != TNumericLimits<double>::Lowest())
+		{
+			ZoomToFitInternal(Axes, InputMin, InputMax, OutputMin, OutputMax);
+		}
+	}
+}
+
+void FCurveEditor::ZoomToFitSelection(EAxisList::Type Axes)
+{
+	if (Selection.Count() <= 1)
+	{
+		ZoomToFit(Axes);
+	}
+	else
+	{
+		TArray<FKeyPosition> KeyPositionsScratch;
+
+		double InputMin = TNumericLimits<double>::Max(), InputMax = TNumericLimits<double>::Lowest(), OutputMin = TNumericLimits<double>::Max(), OutputMax = TNumericLimits<double>::Lowest();
+
+		for (const TTuple<FCurveModelID, FKeyHandleSet>& Pair : Selection.GetAll())
+		{
+			if (const FCurveModel* Curve = FindCurve(Pair.Key))
+			{
+				KeyPositionsScratch.SetNum(Pair.Value.AsArray().Num());
+
+				Curve->GetKeyPositions(Pair.Value.AsArray(), KeyPositionsScratch);
+				for (FKeyPosition Key : KeyPositionsScratch)
+				{
+					InputMin  = FMath::Min(InputMin, Key.InputValue);
+					InputMax  = FMath::Max(InputMax, Key.InputValue);
+					OutputMin = FMath::Min(OutputMin, Key.OutputValue);
+					OutputMax = FMath::Max(OutputMax, Key.OutputValue);
+				}
+			}
+		}
+
+		if (InputMin != TNumericLimits<double>::Max() && InputMax != TNumericLimits<double>::Lowest() && 
+			OutputMin != TNumericLimits<double>::Max() && OutputMax != TNumericLimits<double>::Lowest())
+		{
+			ZoomToFitInternal(Axes, InputMin, InputMax, OutputMin, OutputMax);
+		}
+	}
+}
+
+void FCurveEditor::ZoomToFitInternal(EAxisList::Type Axes, double InputMin, double InputMax, double OutputMin, double OutputMax)
+{
+	FCurveEditorSnapMetrics SnapMetrics = GetSnapMetrics();
+	double MinInputZoom = SnapMetrics.bSnapInputValues ? SnapMetrics.InputSnapRate.AsInterval() : 0.001;
+	double MinOutputZoom = SnapMetrics.bSnapOutputValues ? SnapMetrics.OutputSnapInterval : 0.001;
+
+	InputMax = FMath::Max(InputMin + MinInputZoom, InputMax);
+	OutputMax = FMath::Max(OutputMin + MinOutputZoom, OutputMax);
+
+	double InputPadding  = (InputMax - InputMin) * 0.1;
+	double OutputPadding = (OutputMax - OutputMin) * 0.1;
+
+	InputMin -= InputPadding;
+	InputMax += InputPadding;
+	if (Axes & EAxisList::X)
+	{
+		Bounds->SetInputBounds(InputMin - InputPadding, InputMax + InputPadding);
+	}
+
+	OutputMin -= OutputPadding;
+	OutputMax += OutputPadding;
+	if (Axes & EAxisList::Y)
+	{
+		Bounds->SetOutputBounds(OutputMin - OutputPadding, OutputMax + OutputPadding);
+	}
+}
+
+bool FCurveEditor::IsInputSnappingEnabled() const
+{
+	return InputSnapEnabledAttribute.Get();
+}
+
+void FCurveEditor::ToggleInputSnapping()
+{
+	bool NewValue = !InputSnapEnabledAttribute.Get();
+
+	if (!InputSnapEnabledAttribute.IsBound())
+	{
+		InputSnapEnabledAttribute = NewValue;
+	}
+	else
+	{
+		OnInputSnapEnabledChanged.ExecuteIfBound(NewValue);
+	}
+}
+
+bool FCurveEditor::IsOutputSnappingEnabled() const
+{
+	return OutputSnapEnabledAttribute.Get();
+}
+
+void FCurveEditor::ToggleOutputSnapping()
+{
+	bool NewValue = !OutputSnapEnabledAttribute.Get();
+
+	if (!OutputSnapEnabledAttribute.IsBound())
+	{
+		OutputSnapEnabledAttribute = NewValue;
+	}
+	else
+	{
+		OnOutputSnapEnabledChanged.ExecuteIfBound(NewValue);
+	}
+}
+
+FVector2D FCurveEditor::GetVectorFromSlopeAndLength(float Slope, float Length)
+{
+	float x = Length / FMath::Sqrt(Slope*Slope + 1.f);
+	float y = Slope * x;
+	return FVector2D(x, y);
+}
+
+void FCurveEditor::GetCurveDrawParams(TArray<FCurveDrawParams>& OutDrawParams) const
+{
+	FCurveEditorScreenSpace ScreenSpace = GetScreenSpace();
+	const float DisplayRatio = (ScreenSpace.PixelsPerOutput() / ScreenSpace.PixelsPerInput());
+
+	double InputMin = 0, InputMax = 1;
+	Bounds->GetInputBounds(InputMin, InputMax);
+
+	double OutputMin = 0, OutputMax = 1;
+	Bounds->GetOutputBounds(OutputMin, OutputMax);
+
+	ECurveEditorTangentVisibility::Type TangentVisibility = Settings->GetTangentVisibility();
+
+	OutDrawParams.Reserve(CurveData.Num());
+	for (const TTuple<FCurveModelID, TUniquePtr<FCurveModel>>& Pair : CurveData)
+	{
+		const FKeyHandleSet* SelectedKeys = Selection.GetAll().Find(Pair.Key);
+
+		FCurveDrawParams Params(Pair.Key);
+		Params.Color = Pair.Value->GetColor();
+
+		// Gather the display metrics to use for each key type
+		Pair.Value->GetKeyDrawInfo(ECurvePointType::Key, Params.KeyDrawInfo);
+		Pair.Value->GetKeyDrawInfo(ECurvePointType::ArriveTangent, Params.ArriveTangentDrawInfo);
+		Pair.Value->GetKeyDrawInfo(ECurvePointType::LeaveTangent, Params.LeaveTangentDrawInfo);
+
+		// Gather the interpolating points in input/output space
+		TArray<TTuple<double, double>> InterpolatingPoints;
+
+		Pair.Value->DrawCurve(*this, InterpolatingPoints);
+		Params.InterpolatingPoints.Reserve(InterpolatingPoints.Num());
+
+		double InputOffset = Pair.Value->GetInputDisplayOffset();
+
+		// Convert the interpolating points to screen space
+		for (TTuple<double, double> Point : InterpolatingPoints)
+		{
+			Params.InterpolatingPoints.Add(
+				FVector2D(
+					ScreenSpace.SecondsToScreen(Point.Get<0>() + InputOffset),
+					ScreenSpace.ValueToScreen(Point.Get<1>())
+				)
+			);
+		}
+
+		TArray<FKeyHandle> VisibleKeys;
+		Pair.Value->GetKeys(*this, InputMin, InputMax, OutputMin, OutputMax, VisibleKeys);
+
+		if (VisibleKeys.Num())
+		{
+			TArray<FKeyPosition> AllKeyPositions;
+			TArray<FKeyAttributes> AllKeyAttributes;
+
+			AllKeyPositions.SetNum(VisibleKeys.Num());
+			AllKeyAttributes.SetNum(VisibleKeys.Num());
+
+			Pair.Value->GetKeyPositions(VisibleKeys, AllKeyPositions);
+			Pair.Value->GetKeyAttributes(VisibleKeys, AllKeyAttributes);
+
+			for (int32 Index = 0; Index < VisibleKeys.Num(); ++Index)
+			{
+				const FKeyHandle      KeyHandle   = VisibleKeys[Index];
+				const FKeyPosition&   KeyPosition = AllKeyPositions[Index];
+				const FKeyAttributes& Attributes  = AllKeyAttributes[Index];
+
+				bool bShowTangents = TangentVisibility == ECurveEditorTangentVisibility::AllTangents || ( TangentVisibility == ECurveEditorTangentVisibility::SelectedKeys && SelectedKeys && SelectedKeys->Contains(VisibleKeys[Index]) );
+
+				// Add this key
+				FCurvePointInfo Key(KeyHandle);
+				Key.ScreenPosition = FVector2D(ScreenSpace.SecondsToScreen(KeyPosition.InputValue + InputOffset), ScreenSpace.ValueToScreen(KeyPosition.OutputValue));
+				Key.LayerBias = 2;
+
+				Params.Points.Add(Key);
+
+				if (bShowTangents && Attributes.HasArriveTangent())
+				{
+					float ArriveTangent = Attributes.GetArriveTangent();
+
+					FCurvePointInfo ArriveTangentPoint(KeyHandle);
+					ArriveTangentPoint.Type = ECurvePointType::ArriveTangent;
+					ArriveTangentPoint.ScreenPosition = Key.ScreenPosition + FCurveEditor::GetVectorFromSlopeAndLength(ArriveTangent * -DisplayRatio, -60.f);
+					ArriveTangentPoint.LineDelta = Key.ScreenPosition - ArriveTangentPoint.ScreenPosition;
+					ArriveTangentPoint.LayerBias = 1;
+
+					Params.Points.Add(ArriveTangentPoint);
+				}
+
+				if (bShowTangents && Attributes.HasLeaveTangent())
+				{
+					float LeaveTangent = Attributes.GetLeaveTangent();
+
+					FCurvePointInfo LeaveTangentPoint(KeyHandle);
+					LeaveTangentPoint.Type = ECurvePointType::LeaveTangent;
+					LeaveTangentPoint.ScreenPosition = Key.ScreenPosition + FCurveEditor::GetVectorFromSlopeAndLength(LeaveTangent * -DisplayRatio, 60.f);
+					LeaveTangentPoint.LineDelta = Key.ScreenPosition - LeaveTangentPoint.ScreenPosition;
+					LeaveTangentPoint.LayerBias = 1;
+
+					Params.Points.Add(LeaveTangentPoint);
+				}
+			}
+		}
+
+		OutDrawParams.Add(MoveTemp(Params));
+	}
+}
+
+void FCurveEditor::ConstructInputGridLines(TArray<float>& MajorGridLines, TArray<float>& MinorGridLines) const
+{
+	FCurveEditorScreenSpace ScreenSpace = GetScreenSpace();
+
+	double MajorGridStep  = 0.0;
+	int32  MinorDivisions = 0;
+	if (GetSnapMetrics().InputSnapRate.ComputeGridSpacing(ScreenSpace.PixelsPerInput(), MajorGridStep, MinorDivisions))
+	{
+		const double FirstMajorLine = FMath::FloorToDouble(ScreenSpace.GetInputMin() / MajorGridStep) * MajorGridStep;
+		const double LastMajorLine  = FMath::CeilToDouble(ScreenSpace.GetInputMax() / MajorGridStep) * MajorGridStep;
+
+		for (double CurrentMajorLine = FirstMajorLine; CurrentMajorLine < LastMajorLine; CurrentMajorLine += MajorGridStep)
+		{
+			MajorGridLines.Add( ScreenSpace.SecondsToScreen(CurrentMajorLine) );
+
+			for (int32 Step = 1; Step < MinorDivisions; ++Step)
+			{
+				MinorGridLines.Add( ScreenSpace.SecondsToScreen(CurrentMajorLine + Step*MajorGridStep/MinorDivisions) );
+			}
+		}
+	}
+}
+
+void FCurveEditor::ConstructOutputGridLines(TArray<float>& MajorGridLines, TArray<float>& MinorGridLines, uint8 MinorDivisions) const
+{
+	FCurveEditorSnapMetrics SnapMetrics = GetSnapMetrics();
+	FCurveEditorScreenSpace ScreenSpace = GetScreenSpace();
+
+	if (ScreenSpace.GetOutputMin() == ScreenSpace.GetOutputMax() || ScreenSpace.PixelsPerOutput() <= 0)
+	{
+		return;
+	}
+
+	const float GridPixelSpacing = FMath::Min(ScreenSpace.GetPhysicalHeight()/1.5f, 150.0f);
+
+	static float Base = 10;
+	float MaxTimeStep = GridPixelSpacing / ScreenSpace.PixelsPerOutput();
+	float MajorGridStep = FMath::Pow(Base, FMath::FloorToFloat(FMath::LogX(Base, MaxTimeStep)));
+
+	const double FirstMajorLine = FMath::FloorToDouble(ScreenSpace.GetOutputMin() / MajorGridStep) * MajorGridStep;
+	const double LastMajorLine  = FMath::CeilToDouble(ScreenSpace.GetOutputMax() / MajorGridStep) * MajorGridStep;
+
+	for (double CurrentMajorLine = FirstMajorLine; CurrentMajorLine < LastMajorLine; CurrentMajorLine += MajorGridStep)
+	{
+		MajorGridLines.Add( ScreenSpace.ValueToScreen(CurrentMajorLine) );
+
+		for (int32 Step = 1; Step < MinorDivisions; ++Step)
+		{
+			MinorGridLines.Add( ScreenSpace.SecondsToScreen(CurrentMajorLine + Step*MajorGridStep/MinorDivisions) );
+		}
+	}
+}
+
+void FCurveEditor::DeleteSelection()
+{
+	FScopedTransaction Transaction(LOCTEXT("DeleteKeys", "Delete Keys"));
+
+	for (const TTuple<FCurveModelID, FKeyHandleSet>& Pair : Selection.GetAll())
+	{
+		if (FCurveModel* Curve = FindCurve(Pair.Key))
+		{
+			Curve->Modify();
+			Curve->RemoveKeys(Pair.Value.AsArray());
+		}
+	}
+
+	Selection.Clear();
+}
+
+void FCurveEditor::FlattenSelection()
+{
+	FScopedTransaction Transaction(LOCTEXT("FlattenTangents", "Flatten Tangents"));
+	bool bFoundAnyTangents = false;
+
+	TArray<FKeyHandle> KeyHandles;
+	TArray<FKeyAttributes> AllKeyPositions;
+
+	for (const TTuple<FCurveModelID, FKeyHandleSet>& Pair : Selection.GetAll())
+	{
+		if (FCurveModel* Curve = FindCurve(Pair.Key))
+		{
+			KeyHandles.Reset(Pair.Value.Num());
+			KeyHandles.Append(Pair.Value.AsArray().GetData(), Pair.Value.Num());
+
+			AllKeyPositions.SetNum(KeyHandles.Num());
+			Curve->GetKeyAttributes(KeyHandles, AllKeyPositions);
+
+			// Straighten tangents, ignoring any keys that we can't set tangents on
+			for (int32 Index = AllKeyPositions.Num()-1 ; Index >= 0; --Index)
+			{
+				FKeyAttributes& Attributes = AllKeyPositions[Index];
+				if (Attributes.HasTangentMode() && (Attributes.HasArriveTangent() || Attributes.HasLeaveTangent()))
+				{
+					Attributes.SetArriveTangent(0.f).SetLeaveTangent(0.f);
+					if (Attributes.GetTangentMode() == RCTM_Auto)
+					{
+						Attributes.SetTangentMode(RCTM_User);
+					}
+				}
+				else
+				{
+					AllKeyPositions.RemoveAtSwap(Index, 1, false);
+					KeyHandles.RemoveAtSwap(Index, 1, false);
+				}
+			}
+
+			if (AllKeyPositions.Num() > 0)
+			{
+				Curve->Modify();
+				Curve->SetKeyAttributes(KeyHandles, AllKeyPositions);
+				bFoundAnyTangents = true;
+			}
+		}
+	}
+
+	if (!bFoundAnyTangents)
+	{
+		Transaction.Cancel();
+	}
+}
+
+void FCurveEditor::StraightenSelection()
+{
+	FScopedTransaction Transaction(LOCTEXT("StraightenTangents", "Straighten Tangents"));
+	bool bFoundAnyTangents = false;
+
+	TArray<FKeyHandle> KeyHandles;
+	TArray<FKeyAttributes> AllKeyPositions;
+
+	for (const TTuple<FCurveModelID, FKeyHandleSet>& Pair : Selection.GetAll())
+	{
+		if (FCurveModel* Curve = FindCurve(Pair.Key))
+		{
+			KeyHandles.Reset(Pair.Value.Num());
+			KeyHandles.Append(Pair.Value.AsArray().GetData(), Pair.Value.Num());
+
+			AllKeyPositions.SetNum(KeyHandles.Num());
+			Curve->GetKeyAttributes(KeyHandles, AllKeyPositions);
+
+			// Straighten tangents, ignoring any keys that we can't set tangents on
+			for (int32 Index = AllKeyPositions.Num()-1 ; Index >= 0; --Index)
+			{
+				FKeyAttributes& Attributes = AllKeyPositions[Index];
+				if (Attributes.HasTangentMode() && Attributes.HasArriveTangent() && Attributes.HasLeaveTangent())
+				{
+					float NewTangent = (Attributes.GetLeaveTangent() + Attributes.GetArriveTangent()) * 0.5f;
+					Attributes.SetArriveTangent(NewTangent).SetLeaveTangent(NewTangent);
+					if (Attributes.GetTangentMode() == RCTM_Auto)
+					{
+						Attributes.SetTangentMode(RCTM_User);
+					}
+				}
+				else
+				{
+					AllKeyPositions.RemoveAtSwap(Index, 1, false);
+					KeyHandles.RemoveAtSwap(Index, 1, false);
+				}
+			}
+
+			if (AllKeyPositions.Num() > 0)
+			{
+				Curve->Modify();
+				Curve->SetKeyAttributes(KeyHandles, AllKeyPositions);
+				bFoundAnyTangents = true;
+			}
+		}
+	}
+
+	if (!bFoundAnyTangents)
+	{
+		Transaction.Cancel();
+	}
+}
+
+void FCurveEditor::BakeSelection()
+{
+	FText TransactionText = FText::Format(LOCTEXT("BakeCurves", "Bake {0}|plural(one=Curve, other=Curves)"), Selection.GetAll().Num());
+	FScopedTransaction Transaction(TransactionText);
+
+	TArray<FKeyHandle> KeyHandles;
+	TArray<FKeyPosition> SelectedKeyPositions;
+
+	TArray<FKeyPosition> NewKeyPositions;
+	TArray<FKeyAttributes> NewKeyAttributes;
+
+	FFrameRate BakeRate = GetSnapMetrics().InputSnapRate;
+
+	for (const TTuple<FCurveModelID, FKeyHandleSet>& Pair : Selection.GetAll())
+	{
+		if (FCurveModel* Curve = FindCurve(Pair.Key))
+		{
+			KeyHandles.Reset(Pair.Value.Num());
+			KeyHandles.Append(Pair.Value.AsArray().GetData(), Pair.Value.Num());
+
+			// Get all the selected keys
+			SelectedKeyPositions.SetNum(KeyHandles.Num());
+			Curve->GetKeyPositions(KeyHandles, SelectedKeyPositions);
+
+			// Find the hull of the range of the selected keys
+			double MinKey = TNumericLimits<double>::Max(), MaxKey = TNumericLimits<double>::Lowest();
+			for (FKeyPosition Key : SelectedKeyPositions)
+			{
+				MinKey = FMath::Min(Key.InputValue, MinKey);
+				MaxKey = FMath::Max(Key.InputValue, MaxKey);
+			}
+
+			// Get all keys that exist between the time range
+			KeyHandles.Reset();
+			Curve->GetKeys(*this, MinKey, MaxKey, -MAX_dbl, MAX_dbl, KeyHandles);
+
+			if (KeyHandles.Num() > 1)
+			{
+				// Determine new times for new keys
+				double Interval = BakeRate.AsInterval();
+				int32 NumKeysToAdd = FMath::FloorToInt((MaxKey - MinKey) / Interval);
+
+				NewKeyPositions.Reset(NumKeysToAdd);
+				NewKeyAttributes.Reset(NumKeysToAdd);
+
+				for (int32 KeyIndex = 0; KeyIndex < NumKeysToAdd; ++KeyIndex)
+				{
+					FKeyPosition CurrentKey(MinKey + KeyIndex*Interval, 0.0);
+					if (Curve->Evaluate(CurrentKey.InputValue, CurrentKey.OutputValue))
+					{
+						NewKeyPositions.Add(CurrentKey);
+						NewKeyAttributes.Add(FKeyAttributes().SetInterpMode(RCIM_Linear));
+					}
+				}
+
+				Curve->Modify();
+
+				// Remove all the old keys and add the new ones
+				Curve->RemoveKeys(KeyHandles);
+				Curve->AddKeys(NewKeyPositions, NewKeyAttributes);
+			}
+		}
+	}
+}
+
+/**
+  The following key reduction is the same as that found in FRichCurve.
+  It would be nice if there was just one implementation of the reduction (and) baking algorithms.
+*/
+/** Util to find float value on bezier defined by 4 control points */
+static float BezierInterp(float P0, float P1, float P2, float P3, float Alpha)
+{
+	const float P01 = FMath::Lerp(P0, P1, Alpha);
+	const float P12 = FMath::Lerp(P1, P2, Alpha);
+	const float P23 = FMath::Lerp(P2, P3, Alpha);
+	const float P012 = FMath::Lerp(P01, P12, Alpha);
+	const float P123 = FMath::Lerp(P12, P23, Alpha);
+	const float P0123 = FMath::Lerp(P012, P123, Alpha);
+
+	return P0123;
+}
+
+static float EvalForTwoKeys (const FKeyPosition& Key1Pos, const FKeyAttributes& Key1Attrib,
+							  const FKeyPosition& Key2Pos, const FKeyAttributes& Key2Attrib,
+							  const float InTime)
+{
+	const float Diff = Key2Pos.InputValue - Key1Pos.InputValue;
+
+	if (Diff > 0.f && Key1Attrib.GetInterpMode() != RCIM_Constant)
+	{
+		const float Alpha = (InTime - Key1Pos.InputValue) / Diff;
+		const float P0 = Key1Pos.OutputValue;
+		const float P3 = Key2Pos.OutputValue;
+
+		if (Key1Attrib.GetInterpMode() == RCIM_Linear)
+		{
+			return FMath::Lerp(P0, P3, Alpha);
+		}
+		else
+		{
+			const float OneThird = 1.0f / 3.0f;
+			const float P1 = P0 + (Key1Attrib.GetLeaveTangent() * Diff*OneThird);
+			const float P2 = P3 - (Key2Attrib.GetArriveTangent() * Diff*OneThird);
+
+			return BezierInterp(P0, P1, P2, P3, Alpha);
+		}
+	}
+	else
+	{
+		return Key1Pos.OutputValue;
+	}
+}
+
+
+void FCurveEditor::SimplifySelection(float Tolerance)
+{
+	FText TransactionText = FText::Format(LOCTEXT("SimplifyCurves", "Simplify {0}|plural(one=Curve, other=Curves)"), Selection.GetAll().Num());
+	FScopedTransaction Transaction(TransactionText);
+	TArray<FKeyHandle> KeyHandles;
+	TArray<FKeyPosition> SelectedKeyPositions;
+	TArray<FKeyAttributes> SelectedKeyAttributes;
+
+	FFrameRate BakeRate = GetSnapMetrics().InputSnapRate;
+
+	for (const TTuple<FCurveModelID, FKeyHandleSet>& Pair : Selection.GetAll())
+	{
+		if (FCurveModel* Curve = FindCurve(Pair.Key))
+		{
+			KeyHandles.Reset(Pair.Value.Num());
+			KeyHandles.Append(Pair.Value.AsArray().GetData(), Pair.Value.Num());
+
+			// Get all the selected keys
+			SelectedKeyPositions.SetNum(KeyHandles.Num());
+			Curve->GetKeyPositions(KeyHandles, SelectedKeyPositions);
+
+			// Find the hull of the range of the selected keys
+			double MinKey = TNumericLimits<double>::Max(), MaxKey = TNumericLimits<double>::Lowest();
+			for (FKeyPosition Key : SelectedKeyPositions)
+			{
+				MinKey = FMath::Min(Key.InputValue, MinKey);
+				MaxKey = FMath::Max(Key.InputValue, MaxKey);
+			}
+
+			// Get all keys that exist between the time range
+			KeyHandles.Reset();
+			Curve->GetKeys(*this, MinKey, MaxKey, -MAX_dbl, MAX_dbl, KeyHandles);
+			if (KeyHandles.Num() > 2) //need at least 3 keys to reduce
+			{
+				SelectedKeyPositions.SetNum(KeyHandles.Num());
+				Curve->GetKeyPositions(KeyHandles, SelectedKeyPositions);
+				SelectedKeyAttributes.SetNum(KeyHandles.Num());
+				Curve->GetKeyAttributes(KeyHandles, SelectedKeyAttributes);
+			
+				int32 MostRecentKeepKeyIndex = 0;
+				TArray<FKeyHandle> KeysToRemove;
+				for (int32 TestIndex = 1; TestIndex < KeyHandles.Num()-1; ++TestIndex)
+				{
+					
+					const float KeyValue = SelectedKeyPositions[TestIndex].OutputValue;
+					const float ValueWithoutKey = EvalForTwoKeys(SelectedKeyPositions[MostRecentKeepKeyIndex], SelectedKeyAttributes[MostRecentKeepKeyIndex],
+						SelectedKeyPositions[TestIndex + 1], SelectedKeyAttributes[TestIndex + 1],
+						SelectedKeyPositions[TestIndex].InputValue);
+					if (FMath::Abs(ValueWithoutKey - KeyValue) > Tolerance) // Is this key needed
+					{
+						MostRecentKeepKeyIndex = TestIndex;
+					}
+					else
+					{
+						KeysToRemove.Add(KeyHandles[TestIndex]);
+					}
+					}
+				Curve->Modify();
+				Curve->RemoveKeys(KeysToRemove);
+			}
+		}
+	}
+}
+
+#undef LOCTEXT_NAMESPACE
diff --git a/Engine/Source/Editor/CurveEditor/Private/CurveEditorCommands.cpp b/Engine/Source/Editor/CurveEditor/Private/CurveEditorCommands.cpp
new file mode 100644
--- /dev/null
+++ b/Engine/Source/Editor/CurveEditor/Private/CurveEditorCommands.cpp
@@ -0,0 +1,53 @@
+// Copyright 1998-2018 Epic Games, Inc. All Rights Reserved.
+
+#include "CurveEditorCommands.h"
+
+#define LOCTEXT_NAMESPACE "CurveEditorCommands"
+
+void FCurveEditorCommands::RegisterCommands()
+{
+	UI_COMMAND(ZoomToFitHorizontal, "Fit Horizontal", "Zoom to Fit - Horizontal", EUserInterfaceActionType::Button, FInputChord());
+	UI_COMMAND(ZoomToFitVertical, "Fit Vertical", "Zoom to Fit - Vertical", EUserInterfaceActionType::Button, FInputChord());
+	UI_COMMAND(ZoomToFit, "Fit", "Zoom to Fit", EUserInterfaceActionType::Button, FInputChord(EKeys::F));
+	UI_COMMAND(ZoomToFitAll, "FitAll", "Zoom to Fit All", EUserInterfaceActionType::Button, FInputChord(EKeys::A));
+
+	UI_COMMAND(ToggleInputSnapping, "Input Snapping", "Toggle Input Snapping", EUserInterfaceActionType::ToggleButton, FInputChord());
+	UI_COMMAND(ToggleOutputSnapping, "Output Snapping", "Toggle Output Snapping", EUserInterfaceActionType::ToggleButton, FInputChord());
+
+	UI_COMMAND(InterpolationConstant, "Constant", "Constant interpolation", EUserInterfaceActionType::ToggleButton, FInputChord(EKeys::Five));
+	UI_COMMAND(InterpolationLinear, "Linear", "Linear interpolation", EUserInterfaceActionType::ToggleButton, FInputChord(EKeys::Four));
+	UI_COMMAND(InterpolationCubicAuto, "Auto", "Cubic interpolation - Automatic tangents", EUserInterfaceActionType::ToggleButton, FInputChord(EKeys::One));
+	UI_COMMAND(InterpolationCubicUser, "User", "Cubic interpolation - User flat tangents", EUserInterfaceActionType::ToggleButton, FInputChord(EKeys::Two));
+	UI_COMMAND(InterpolationCubicBreak, "Break", "Cubic interpolation - User broken tangents", EUserInterfaceActionType::ToggleButton, FInputChord(EKeys::Three));
+
+	UI_COMMAND(FlattenTangents, "Flatten", "Flatten tangents", EUserInterfaceActionType::Button, FInputChord());
+	UI_COMMAND(StraightenTangents, "Straighten", "Straighten tangents", EUserInterfaceActionType::Button, FInputChord());
+
+	UI_COMMAND(BakeCurve, "Bake", "Bake curve", EUserInterfaceActionType::Button, FInputChord());
+	UI_COMMAND(ReduceCurve, "Reduce", "Reduce curve", EUserInterfaceActionType::Button, FInputChord());
+
+	UI_COMMAND(SetPreInfinityExtrapCycle, "Cycle", "Set the pre-infinity curve extrapolation to cycle.", EUserInterfaceActionType::RadioButton, FInputChord());
+	UI_COMMAND(SetPreInfinityExtrapCycleWithOffset, "Cycle with Offset", "Set the pre-infinity curve extrapolation to cycle with offset.", EUserInterfaceActionType::RadioButton, FInputChord());
+	UI_COMMAND(SetPreInfinityExtrapOscillate, "Oscillate", "Set the pre-infinity curve extrapolation to oscillate.", EUserInterfaceActionType::RadioButton, FInputChord());
+	UI_COMMAND(SetPreInfinityExtrapLinear, "Linear", "Set the pre-infinity curve extrapolation to linear.", EUserInterfaceActionType::RadioButton, FInputChord());
+	UI_COMMAND(SetPreInfinityExtrapConstant, "Constant", "Set the pre-infinity curve extrapolation to constant.", EUserInterfaceActionType::RadioButton, FInputChord());
+	UI_COMMAND(SetPostInfinityExtrapCycle, "Cycle", "Set the post-infinity curve extrapolation to cycle.", EUserInterfaceActionType::RadioButton, FInputChord());
+	UI_COMMAND(SetPostInfinityExtrapCycleWithOffset, "Cycle with Offset", "Set the post-infinity curve extrapolation to cycle with offset.", EUserInterfaceActionType::RadioButton, FInputChord());
+	UI_COMMAND(SetPostInfinityExtrapOscillate, "Oscillate", "Set the post-infinity curve extrapolation to oscillate.", EUserInterfaceActionType::RadioButton, FInputChord());
+	UI_COMMAND(SetPostInfinityExtrapLinear, "Linear", "Set the post-infinity curve extrapolation to linear.", EUserInterfaceActionType::RadioButton, FInputChord());
+	UI_COMMAND(SetPostInfinityExtrapConstant, "Constant", "Set the post-infinity curve extrapolation to constant.", EUserInterfaceActionType::RadioButton, FInputChord());
+
+	UI_COMMAND(SetAllTangentsVisibility, "All Tangents", "Show all tangents in the curve editor.", EUserInterfaceActionType::RadioButton, FInputChord());
+	UI_COMMAND(SetSelectedKeysTangentVisibility, "Selected Keys", "Show tangents for selected keys in the curve editor.", EUserInterfaceActionType::RadioButton, FInputChord());
+	UI_COMMAND(SetNoTangentsVisibility, "No Tangents", "Show no tangents in the curve editor.", EUserInterfaceActionType::RadioButton, FInputChord());
+
+	UI_COMMAND(ToggleAutoFrameCurveEditor, "Auto Frame Curves", "Auto frame curves when they are selected.", EUserInterfaceActionType::ToggleButton, FInputChord() );
+	UI_COMMAND(ToggleShowCurveEditorCurveToolTips, "Curve Tool Tips", "Show a tool tip with name and values when hovering over a curve.", EUserInterfaceActionType::ToggleButton, FInputChord() );
+
+	UI_COMMAND(AddKeyHovered, "Add Key", "Add a new key to this curve at the current position.", EUserInterfaceActionType::Button, FInputChord(EKeys::MiddleMouseButton) );
+
+	UI_COMMAND(AddKeyToAllCurves, "Add Key", "Add a new key to all curves at the current time.", EUserInterfaceActionType::Button, FInputChord(EModifierKey::Shift, EKeys::LeftMouseButton) );
+	UI_COMMAND(AddKeyToAllCurvesHere, "Add Key Here", "Add a new key to all curves at the current time and value.", EUserInterfaceActionType::Button, FInputChord(EModifierKey::Shift | EModifierKey::Control, EKeys::LeftMouseButton) );
+}
+
+#undef LOCTEXT_NAMESPACE
diff --git a/Engine/Source/Editor/CurveEditor/Private/CurveEditorContextMenu.cpp b/Engine/Source/Editor/CurveEditor/Private/CurveEditorContextMenu.cpp
new file mode 100644
--- /dev/null
+++ b/Engine/Source/Editor/CurveEditor/Private/CurveEditorContextMenu.cpp
@@ -0,0 +1,89 @@
+// Copyright 1998-2018 Epic Games, Inc. All Rights Reserved.
+
+#include "CurveEditorContextMenu.h"
+#include "Framework/MultiBox/MultiBoxBuilder.h"
+#include "Framework/Commands/GenericCommands.h"
+#include "CurveEditor.h"
+#include "CurveEditorCommands.h"
+
+#define LOCTEXT_NAMESPACE "CurveEditorContextMenu"
+
+void FCurveEditorContextMenu::BuildMenu(FMenuBuilder& MenuBuilder, TWeakPtr<FCurveEditor> WeakCurveEditor, TOptional<FCurvePointHandle> ClickedPoint, TOptional<FCurveModelID> HoveredCurveID)
+{
+	TSharedPtr<FCurveEditor> CurveEditor = WeakCurveEditor.Pin();
+	if (!CurveEditor.IsValid())
+	{
+		return;
+	}
+
+	int32 NumSelectedKeys = CurveEditor->Selection.Count();
+	if (NumSelectedKeys > 0)
+	{
+		MenuBuilder.BeginSection("CurveEditorKeySection", FText::Format(LOCTEXT("CurveEditorKeySection", "{0} Selected {0}|plural(one=Key,other=Keys)"), NumSelectedKeys));
+		{
+			MenuBuilder.AddMenuEntry(FGenericCommands::Get().Delete);
+
+			MenuBuilder.AddMenuEntry(FCurveEditorCommands::Get().FlattenTangents);
+			MenuBuilder.AddMenuEntry(FCurveEditorCommands::Get().StraightenTangents);
+			MenuBuilder.AddMenuEntry(FCurveEditorCommands::Get().ReduceCurve);
+			MenuBuilder.AddMenuEntry(FCurveEditorCommands::Get().BakeCurve);
+
+			MenuBuilder.AddMenuSeparator();
+
+			MenuBuilder.AddMenuEntry(FCurveEditorCommands::Get().InterpolationCubicAuto);
+			MenuBuilder.AddMenuEntry(FCurveEditorCommands::Get().InterpolationCubicUser);
+			MenuBuilder.AddMenuEntry(FCurveEditorCommands::Get().InterpolationCubicBreak);
+			MenuBuilder.AddMenuEntry(FCurveEditorCommands::Get().InterpolationLinear);
+			MenuBuilder.AddMenuEntry(FCurveEditorCommands::Get().InterpolationConstant);
+		}
+		MenuBuilder.EndSection();
+	}
+	else
+	{
+		const FCurveModel* HoveredCurve = HoveredCurveID.IsSet() ? CurveEditor->FindCurve(HoveredCurveID.GetValue()) : nullptr;
+		if (HoveredCurve)
+		{
+			MenuBuilder.BeginSection("CurveEditorCurveSection", FText::Format(LOCTEXT("CurveNameFormat", "Curve '{0}'"), HoveredCurve->GetDisplayName()));
+			{
+				MenuBuilder.AddMenuEntry(FCurveEditorCommands::Get().AddKeyHovered);
+
+				MenuBuilder.AddMenuEntry(FCurveEditorCommands::Get().ReduceCurve);
+				MenuBuilder.AddMenuEntry(FCurveEditorCommands::Get().BakeCurve);
+
+				MenuBuilder.AddSubMenu(LOCTEXT("PreExtrapText", "Pre-Extrap"), FText(), FNewMenuDelegate::CreateLambda(
+					[](FMenuBuilder& SubMenu)
+					{
+						SubMenu.AddMenuEntry(FCurveEditorCommands::Get().SetPreInfinityExtrapCycle);
+						SubMenu.AddMenuEntry(FCurveEditorCommands::Get().SetPreInfinityExtrapCycleWithOffset);
+						SubMenu.AddMenuEntry(FCurveEditorCommands::Get().SetPreInfinityExtrapOscillate);
+						SubMenu.AddMenuEntry(FCurveEditorCommands::Get().SetPreInfinityExtrapLinear);
+						SubMenu.AddMenuEntry(FCurveEditorCommands::Get().SetPreInfinityExtrapConstant);
+					})
+				);
+
+				MenuBuilder.AddSubMenu(LOCTEXT("PostExtrapText", "Post-Extrap"), FText(), FNewMenuDelegate::CreateLambda(
+					[](FMenuBuilder& SubMenu)
+					{
+						SubMenu.AddMenuEntry(FCurveEditorCommands::Get().SetPostInfinityExtrapCycle);
+						SubMenu.AddMenuEntry(FCurveEditorCommands::Get().SetPostInfinityExtrapCycleWithOffset);
+						SubMenu.AddMenuEntry(FCurveEditorCommands::Get().SetPostInfinityExtrapOscillate);
+						SubMenu.AddMenuEntry(FCurveEditorCommands::Get().SetPostInfinityExtrapLinear);
+						SubMenu.AddMenuEntry(FCurveEditorCommands::Get().SetPostInfinityExtrapConstant);
+					})
+				);
+			}
+			MenuBuilder.EndSection();
+		}
+		else
+		{
+			MenuBuilder.BeginSection("CurveEditorAllCurveSections", LOCTEXT("CurveEditorAllCurveSections", "All Curves"));
+			{
+				MenuBuilder.AddMenuEntry(FCurveEditorCommands::Get().AddKeyToAllCurves);
+				MenuBuilder.AddMenuEntry(FCurveEditorCommands::Get().AddKeyToAllCurvesHere);
+			}
+			MenuBuilder.EndSection();
+		}
+	}
+}
+
+#undef LOCTEXT_NAMESPACE
\ No newline at end of file
diff --git a/Engine/Source/Editor/CurveEditor/Private/CurveEditorContextMenu.h b/Engine/Source/Editor/CurveEditor/Private/CurveEditorContextMenu.h
new file mode 100644
--- /dev/null
+++ b/Engine/Source/Editor/CurveEditor/Private/CurveEditorContextMenu.h
@@ -0,0 +1,15 @@
+// Copyright 1998-2018 Epic Games, Inc. All Rights Reserved.
+
+#pragma once
+
+#include "Templates/SharedPointer.h"
+#include "Misc/Optional.h"
+#include "CurveEditorTypes.h"
+
+class FCurveEditor;
+class FMenuBuilder;
+
+struct FCurveEditorContextMenu
+{
+	static void BuildMenu(FMenuBuilder& MenuBuilder, TWeakPtr<FCurveEditor> WeakCurveEditor, TOptional<FCurvePointHandle> ClickedPoint, TOptional<FCurveModelID> HoveredCurve);
+};
\ No newline at end of file
diff --git a/Engine/Source/Editor/CurveEditor/Private/CurveEditorModule.cpp b/Engine/Source/Editor/CurveEditor/Private/CurveEditorModule.cpp
new file mode 100644
--- /dev/null
+++ b/Engine/Source/Editor/CurveEditor/Private/CurveEditorModule.cpp
@@ -0,0 +1,27 @@
+// Copyright 1998-2018 Epic Games, Inc. All Rights Reserved.
+
+#include "Modules/ModuleManager.h"
+#include "CurveEditorCommands.h"
+
+
+class FCurveEditorModule : public IModuleInterface
+{
+public:
+	FCurveEditorModule()
+	{
+		check(false);
+	}
+
+	virtual void StartupModule() override
+	{
+		FCurveEditorCommands::Register();
+	}
+
+	virtual void ShutdownModule() override
+	{
+		FCurveEditorCommands::Unregister();
+	}
+};
+
+
+IMPLEMENT_MODULE(FCurveEditorModule, CurveEditor)
\ No newline at end of file
diff --git a/Engine/Source/Editor/CurveEditor/Private/CurveEditorScreenSpace.cpp b/Engine/Source/Editor/CurveEditor/Private/CurveEditorScreenSpace.cpp
new file mode 100644
--- /dev/null
+++ b/Engine/Source/Editor/CurveEditor/Private/CurveEditorScreenSpace.cpp
@@ -0,0 +1,5 @@
+// Copyright 1998-2018 Epic Games, Inc. All Rights Reserved.
+
+#include "CurveEditorScreenSpace.h"
+#include "Math/Range.h"
+#include "Layout/Geometry.h"
diff --git a/Engine/Source/Editor/CurveEditor/Private/DragOperations/CurveEditorDragOperation_Marquee.cpp b/Engine/Source/Editor/CurveEditor/Private/DragOperations/CurveEditorDragOperation_Marquee.cpp
new file mode 100644
--- /dev/null
+++ b/Engine/Source/Editor/CurveEditor/Private/DragOperations/CurveEditorDragOperation_Marquee.cpp
@@ -0,0 +1,105 @@
+// Copyright 1998-2018 Epic Games, Inc. All Rights Reserved.
+
+#include "DragOperations/CurveEditorDragOperation_Marquee.h"
+#include "CurveEditor.h"
+#include "SCurveEditorPanel.h"
+#include "EditorStyleSet.h"
+#include "CurveDrawInfo.h"
+
+FCurveEditorDragOperation_Marquee::FCurveEditorDragOperation_Marquee(FCurveEditor* InCurveEditor, const SCurveEditorPanel* InCurveEditorPanel)
+{
+	CurveEditor = InCurveEditor;
+	CurveEditorPanel = InCurveEditorPanel;
+}
+
+void FCurveEditorDragOperation_Marquee::OnBeginDrag(FVector2D InitialPosition, FVector2D CurrentPosition, const FPointerEvent& MouseEvent)
+{
+	Marquee = FSlateRect(
+		FMath::Min(InitialPosition.X, CurrentPosition.X),
+		FMath::Min(InitialPosition.Y, CurrentPosition.Y),
+		FMath::Max(InitialPosition.X, CurrentPosition.X),
+		FMath::Max(InitialPosition.Y, CurrentPosition.Y)
+		);
+}
+
+void FCurveEditorDragOperation_Marquee::OnDrag(FVector2D InitialPosition, FVector2D CurrentPosition, const FPointerEvent& MouseEvent)
+{
+	Marquee = FSlateRect(
+		FMath::Min(InitialPosition.X, CurrentPosition.X),
+		FMath::Min(InitialPosition.Y, CurrentPosition.Y),
+		FMath::Max(InitialPosition.X, CurrentPosition.X),
+		FMath::Max(InitialPosition.Y, CurrentPosition.Y)
+		);
+}
+
+void FCurveEditorDragOperation_Marquee::OnEndDrag(FVector2D InitialPosition, FVector2D CurrentPosition, const FPointerEvent& MouseEvent)
+{
+	Marquee = FSlateRect(
+		FMath::Min(InitialPosition.X, CurrentPosition.X),
+		FMath::Min(InitialPosition.Y, CurrentPosition.Y),
+		FMath::Max(InitialPosition.X, CurrentPosition.X),
+		FMath::Max(InitialPosition.Y, CurrentPosition.Y)
+		);
+
+	FCurveEditorScreenSpace ScreenSpace = CurveEditor->GetScreenSpace();
+
+	bool bRemoveFromSelection = MouseEvent.IsAltDown();
+
+	// Only select the same types of point
+	TOptional<ECurvePointType> MatchPointType;
+
+	if (!MouseEvent.IsShiftDown() && !bRemoveFromSelection)
+	{
+		CurveEditor->Selection.Clear();
+	}
+	else if (CurveEditor->Selection.Count() != 0)
+	{
+		MatchPointType = CurveEditor->Selection.GetSelectionType();
+	}
+
+	FSlateRect MarqueeRectPx = Marquee;
+	for (const FCurveDrawParams& DrawParams : CurveEditorPanel->GetCachedDrawParams())
+	{
+		for (const FCurvePointInfo& Point : DrawParams.Points)
+		{
+			// Can we select this type of point?
+			if (MatchPointType.IsSet() && Point.Type != MatchPointType.GetValue())
+			{
+				continue;
+			}
+
+			const FKeyDrawInfo& DrawInfo = DrawParams.GetKeyDrawInfo(Point.Type);
+			FSlateRect PointRect = FSlateRect::FromPointAndExtent(Point.ScreenPosition - DrawInfo.ScreenSize/2, DrawInfo.ScreenSize);
+
+			if (FSlateRect::DoRectanglesIntersect(PointRect, MarqueeRectPx))
+			{
+				if (!MatchPointType.IsSet() && Point.Type == ECurvePointType::Key)
+				{
+					CurveEditor->Selection.Clear();
+					MatchPointType = ECurvePointType::Key;
+				}
+
+				if (bRemoveFromSelection)
+				{
+					CurveEditor->Selection.Remove(DrawParams.GetID(), Point.Type, Point.KeyHandle);
+				}
+				else
+				{
+					CurveEditor->Selection.Add(DrawParams.GetID(), Point.Type, Point.KeyHandle);
+				}
+			}
+		}
+	}
+}
+
+int32 FCurveEditorDragOperation_Marquee::OnPaint(const FGeometry& AllottedGeometry, FSlateWindowElementList& OutDrawElements, int32 LayerId)
+{
+	FSlateDrawElement::MakeBox(
+		OutDrawElements,
+		LayerId,
+		AllottedGeometry.ToPaintGeometry(Marquee.GetTopLeft(), Marquee.GetBottomRight() - Marquee.GetTopLeft()),
+		FEditorStyle::GetBrush(TEXT("MarqueeSelection"))
+		);
+
+	return LayerId;
+}
diff --git a/Engine/Source/Editor/CurveEditor/Private/DragOperations/CurveEditorDragOperation_Marquee.h b/Engine/Source/Editor/CurveEditor/Private/DragOperations/CurveEditorDragOperation_Marquee.h
new file mode 100644
--- /dev/null
+++ b/Engine/Source/Editor/CurveEditor/Private/DragOperations/CurveEditorDragOperation_Marquee.h
@@ -0,0 +1,34 @@
+// Copyright 1998-2018 Epic Games, Inc. All Rights Reserved.
+
+
+#pragma once
+
+#include "CoreTypes.h"
+#include "Math/Vector2D.h"
+#include "Containers/Map.h"
+#include "UObject/NameTypes.h"
+#include "ICurveEditorDragOperation.h"
+
+class FCurveEditor;
+class SCurveEditorPanel;
+
+class FCurveEditorDragOperation_Marquee : public ICurveEditorDragOperation
+{
+public:
+
+	FCurveEditorDragOperation_Marquee(FCurveEditor* InCurveEditor, const SCurveEditorPanel* InCurveEditorPanel);
+
+	virtual void OnBeginDrag(FVector2D InitialPosition, FVector2D CurrentPosition, const FPointerEvent& MouseEvent) override;
+	virtual void OnDrag(FVector2D InitialPosition, FVector2D CurrentPosition, const FPointerEvent& MouseEvent) override;
+	virtual void OnEndDrag(FVector2D InitialPosition, FVector2D CurrentPosition, const FPointerEvent& MouseEvent) override;
+	virtual int32 OnPaint(const FGeometry& AllottedGeometry, FSlateWindowElementList& OutDrawElements, int32 LayerId) override;
+
+private:
+
+	/** The current marquee rectangle */
+	FSlateRect Marquee;
+	/** Ptr back to the curve editor */
+	FCurveEditor* CurveEditor;
+	/** Ptr back to the curve editor panel */
+	const SCurveEditorPanel* CurveEditorPanel;
+};
\ No newline at end of file
diff --git a/Engine/Source/Editor/CurveEditor/Private/DragOperations/CurveEditorDragOperation_MoveKeys.cpp b/Engine/Source/Editor/CurveEditor/Private/DragOperations/CurveEditorDragOperation_MoveKeys.cpp
new file mode 100644
--- /dev/null
+++ b/Engine/Source/Editor/CurveEditor/Private/DragOperations/CurveEditorDragOperation_MoveKeys.cpp
@@ -0,0 +1,86 @@
+// Copyright 1998-2018 Epic Games, Inc. All Rights Reserved.
+
+#include "DragOperations/CurveEditorDragOperation_MoveKeys.h"
+#include "CurveEditorScreenSpace.h"
+#include "CurveEditor.h"
+
+void FCurveEditorDragOperation_MoveKeys::OnInitialize(FCurveEditor* InCurveEditor, const TOptional<FCurvePointHandle>& CardinalPoint)
+{
+	CurveEditor = InCurveEditor;
+}
+
+
+void FCurveEditorDragOperation_MoveKeys::OnBeginDrag(FVector2D InitialPosition, FVector2D CurrentPosition, const FPointerEvent& MouseEvent)
+{
+	int32 NumKeys = CurveEditor->Selection.Count();
+	Transaction = MakeUnique<FScopedTransaction>(FText::Format(NSLOCTEXT("CurveEditor", "MoveKeysFormat", "Move {0}|plural(one=Key, other=Keys)"), NumKeys));
+
+	KeysByCurve.Reset();
+
+	for (const TTuple<FCurveModelID, FKeyHandleSet>& Pair : CurveEditor->Selection.GetAll())
+	{
+		FCurveModelID CurveID = Pair.Key;
+		FCurveModel*  Curve   = CurveEditor->FindCurve(CurveID);
+
+		if (ensure(Curve))
+		{
+			Curve->Modify();
+
+			TArrayView<const FKeyHandle> Handles = Pair.Value.AsArray();
+
+			FKeyData& KeyData = KeysByCurve.Emplace_GetRef(CurveID);
+			KeyData.Handles = TArray<FKeyHandle>(Handles.GetData(), Handles.Num());
+
+			KeyData.StartKeyPositions.SetNumZeroed(KeyData.Handles.Num());
+			Curve->GetKeyPositions(KeyData.Handles, KeyData.StartKeyPositions);
+		}
+	}
+}
+
+void FCurveEditorDragOperation_MoveKeys::OnDrag(FVector2D InitialPosition, FVector2D CurrentPosition, const FPointerEvent& MouseEvent)
+{
+	TArray<FKeyPosition> NewKeyPositionScratch;
+
+	FVector2D MousePosition = GetLockedMousePosition(InitialPosition, CurrentPosition, MouseEvent);
+
+	FCurveEditorScreenSpace ScreenSpace = CurveEditor->GetScreenSpace();
+	double DeltaInput  =  (MousePosition.X - InitialPosition.X) / ScreenSpace.PixelsPerInput();
+	double DeltaOutput = -(MousePosition.Y - InitialPosition.Y) / ScreenSpace.PixelsPerOutput();
+
+	for (const FKeyData& KeyData : KeysByCurve)
+	{
+		FCurveModel* Curve = CurveEditor->FindCurve(KeyData.CurveID);
+
+		if (ensure(Curve))
+		{
+			NewKeyPositionScratch.Reset();
+			NewKeyPositionScratch.Reserve(KeyData.StartKeyPositions.Num());
+
+			for (FKeyPosition StartPosition : KeyData.StartKeyPositions)
+			{
+				StartPosition.InputValue  += DeltaInput;
+				StartPosition.OutputValue += DeltaOutput;
+
+				StartPosition.InputValue  = SnapMetrics.SnapInputSeconds(StartPosition.InputValue);
+				StartPosition.OutputValue = SnapMetrics.SnapOutput(StartPosition.OutputValue);
+
+				NewKeyPositionScratch.Add(StartPosition);
+			}
+
+			Curve->SetKeyPositions(KeyData.Handles, NewKeyPositionScratch);
+		}
+	}
+}
+
+void FCurveEditorDragOperation_MoveKeys::OnCancelDrag()
+{
+	ICurveEditorKeyDragOperation::OnCancelDrag();
+
+	for (const FKeyData& KeyData : KeysByCurve)
+	{
+		if (FCurveModel* Curve = CurveEditor->FindCurve(KeyData.CurveID))
+		{
+			Curve->SetKeyPositions(KeyData.Handles, KeyData.StartKeyPositions);
+		}
+	}
+}
\ No newline at end of file
diff --git a/Engine/Source/Editor/CurveEditor/Private/DragOperations/CurveEditorDragOperation_Pan.cpp b/Engine/Source/Editor/CurveEditor/Private/DragOperations/CurveEditorDragOperation_Pan.cpp
new file mode 100644
--- /dev/null
+++ b/Engine/Source/Editor/CurveEditor/Private/DragOperations/CurveEditorDragOperation_Pan.cpp
@@ -0,0 +1,36 @@
+// Copyright 1998-2018 Epic Games, Inc. All Rights Reserved.
+
+#include "CurveEditorDragOperation_Pan.h"
+#include "CurveEditorScreenSpace.h"
+#include "CurveEditor.h"
+
+FCurveEditorDragOperation_Pan::FCurveEditorDragOperation_Pan(FCurveEditor* InCurveEditor)
+	: CurveEditor(InCurveEditor)
+{
+}
+
+void FCurveEditorDragOperation_Pan::OnBeginDrag(FVector2D InitialPosition, FVector2D CurrentPosition, const FPointerEvent& MouseEvent)
+{
+	FCurveEditorScreenSpace ScreenSpace = CurveEditor->GetScreenSpace();
+
+	InitialInputMin  = ScreenSpace.GetInputMin();
+	InitialInputMax  = ScreenSpace.GetInputMax();
+	InitialOutputMin = ScreenSpace.GetOutputMin();
+	InitialOutputMax = ScreenSpace.GetOutputMax();
+}
+
+void FCurveEditorDragOperation_Pan::OnDrag(FVector2D InitialPosition, FVector2D CurrentPosition, const FPointerEvent& MouseEvent)
+{
+	FVector2D PixelDelta = GetLockedMousePosition(InitialPosition, CurrentPosition, MouseEvent) - InitialPosition;
+
+	FCurveEditorScreenSpace ScreenSpace = CurveEditor->GetScreenSpace();
+
+	double InputMin  = InitialInputMin - PixelDelta.X / ScreenSpace.PixelsPerInput();
+	double InputMax  = InitialInputMax - PixelDelta.X / ScreenSpace.PixelsPerInput();
+
+	double OutputMin = InitialOutputMin + PixelDelta.Y / ScreenSpace.PixelsPerOutput();
+	double OutputMax = InitialOutputMax + PixelDelta.Y / ScreenSpace.PixelsPerOutput();
+
+	CurveEditor->GetBounds().SetInputBounds(InputMin, InputMax);
+	CurveEditor->GetBounds().SetOutputBounds(OutputMin, OutputMax);
+}
\ No newline at end of file
diff --git a/Engine/Source/Editor/CurveEditor/Private/DragOperations/CurveEditorDragOperation_Pan.h b/Engine/Source/Editor/CurveEditor/Private/DragOperations/CurveEditorDragOperation_Pan.h
new file mode 100644
--- /dev/null
+++ b/Engine/Source/Editor/CurveEditor/Private/DragOperations/CurveEditorDragOperation_Pan.h
@@ -0,0 +1,22 @@
+// Copyright 1998-2018 Epic Games, Inc. All Rights Reserved.
+
+#pragma once
+
+#include "ICurveEditorDragOperation.h"
+
+class FCurveEditor;
+
+class FCurveEditorDragOperation_Pan : public ICurveEditorDragOperation
+{
+public:
+	FCurveEditorDragOperation_Pan(FCurveEditor* CurveEditor);
+
+	virtual void OnBeginDrag(FVector2D InitialPosition, FVector2D CurrentPosition, const FPointerEvent& MouseEvent) override;
+	virtual void OnDrag(FVector2D InitialPosition, FVector2D CurrentPosition, const FPointerEvent& MouseEvent) override;
+
+private:
+
+	FCurveEditor* CurveEditor;
+	double InitialInputMin, InitialInputMax;
+	double InitialOutputMin, InitialOutputMax;
+};
\ No newline at end of file
diff --git a/Engine/Source/Editor/CurveEditor/Private/DragOperations/CurveEditorDragOperation_Tangent.cpp b/Engine/Source/Editor/CurveEditor/Private/DragOperations/CurveEditorDragOperation_Tangent.cpp
new file mode 100644
--- /dev/null
+++ b/Engine/Source/Editor/CurveEditor/Private/DragOperations/CurveEditorDragOperation_Tangent.cpp
@@ -0,0 +1,178 @@
+// Copyright 1998-2018 Epic Games, Inc. All Rights Reserved.
+
+#include "CurveEditorDragOperation_Tangent.h"
+#include "CurveEditorScreenSpace.h"
+#include "CurveEditor.h"
+
+void FCurveEditorDragOperation_Tangent::OnInitialize(FCurveEditor* InCurveEditor, const TOptional<FCurvePointHandle>& InCardinalPoint)
+{
+	CurveEditor = InCurveEditor;
+
+	// Find the position of the clicked on key (no type name)
+	if (InCardinalPoint.IsSet())
+	{
+		if (const FCurveModel* Curve = CurveEditor->FindCurve(InCardinalPoint->CurveID))
+		{
+			FCurveEditorScreenSpace ScreenSpace = CurveEditor->GetScreenSpace();
+
+			FKeyPosition CardinalKeyPosition;
+
+			Curve->GetKeyPositions(
+				TArrayView<const FKeyHandle>(&InCardinalPoint->KeyHandle, 1),
+				MakeArrayView(&CardinalKeyPosition, 1));
+
+			CardinalPosition = FVector2D(ScreenSpace.SecondsToScreen(CardinalKeyPosition.InputValue), ScreenSpace.ValueToScreen(CardinalKeyPosition.OutputValue));
+		}
+	}
+}
+
+void FCurveEditorDragOperation_Tangent::OnBeginDrag(FVector2D InitialPosition, FVector2D CurrentPosition, const FPointerEvent& MouseEvent)
+{
+	PointType = CurveEditor->Selection.GetSelectionType();
+
+	int32 NumKeys = CurveEditor->Selection.Count();
+
+	FText Description = PointType == ECurvePointType::ArriveTangent
+		? FText::Format(NSLOCTEXT("CurveEditor", "DragEntryTangentsFormat", "Drag Entry {0}|plural(one=Tangent, other=Tangents)"), NumKeys)
+		: FText::Format(NSLOCTEXT("CurveEditor", "DragExitTangentsFormat", "Drag Exit {0}|plural(one=Tangent, other=Tangents)"), NumKeys);
+
+	Transaction = MakeUnique<FScopedTransaction>(Description);
+
+	KeysByCurve.Reset();
+	for (const TTuple<FCurveModelID, FKeyHandleSet>& Pair : CurveEditor->Selection.GetAll())
+	{
+		FCurveModelID CurveID = Pair.Key;
+		FCurveModel* Curve    = CurveEditor->FindCurve(CurveID);
+
+		if (ensure(Curve))
+		{
+			Curve->Modify();
+
+			TArrayView<const FKeyHandle> Handles = Pair.Value.AsArray();
+
+			FKeyData& KeyData = KeysByCurve.Emplace_GetRef(CurveID);
+			KeyData.Handles = TArray<FKeyHandle>(Handles.GetData(), Handles.Num());
+
+			KeyData.Attributes.SetNum(KeyData.Handles.Num());
+			Curve->GetKeyAttributes(KeyData.Handles, KeyData.Attributes);
+		}
+	}
+}
+
+void FCurveEditorDragOperation_Tangent::OnDrag(FVector2D InitialPosition, FVector2D CurrentPosition, const FPointerEvent& MouseEvent)
+{
+	FCurveEditorScreenSpace ScreenSpace = CurveEditor->GetScreenSpace();
+	const float DisplayRatio = (ScreenSpace.PixelsPerOutput() / ScreenSpace.PixelsPerInput());
+
+	TOptional<float> NewTangent;
+	
+	FVector2D PixelDelta = CurrentPosition - CardinalPosition.Get(InitialPosition);
+
+	if (MouseEvent.IsShiftDown())
+	{
+		PixelDelta = RoundTrajectory(PixelDelta);
+	}
+
+	if (PointType == ECurvePointType::ArriveTangent)
+	{
+		PixelDelta.X = FMath::Min(PixelDelta.X, -KINDA_SMALL_NUMBER);
+	}
+	else
+	{
+		PixelDelta.X = FMath::Max(PixelDelta.X, KINDA_SMALL_NUMBER);
+	}
+	//If the CardinalPosition is set that means we are doing a left mouse and we are over the key(cardinal) position
+	//and we can set the tangent directly.
+	//If we are doing a middle mouse move, there is no cardinal position, but rather we need to apply the
+	//PixelDelta invidually to each tangent
+	if (CardinalPosition.IsSet())
+	{
+		NewTangent = (-PixelDelta.Y / PixelDelta.X) / DisplayRatio;
+	}
+
+	TArray<FKeyAttributes> NewKeyAttributesScratch;
+	TArray<FKeyPosition> KeyPositions;
+
+	for (const FKeyData& KeyData : KeysByCurve)
+	{
+		FCurveModel* Curve = CurveEditor->FindCurve(KeyData.CurveID);
+		if (ensure(Curve))
+		{
+			NewKeyAttributesScratch.Reset();
+			NewKeyAttributesScratch.Reserve(KeyData.Attributes.Num());
+			if (!NewTangent.IsSet()) //No tangent, so MMB, so cache out positions for each key
+			{
+				KeyPositions.SetNumUninitialized(KeyData.Attributes.Num());
+				Curve->GetKeyPositions(KeyData.Handles, KeyPositions);
+			}
+
+			if (PointType == ECurvePointType::ArriveTangent)
+			{
+				for(int32 i =0; i < KeyData.Handles.Num();++i)
+				{
+					if (NewTangent.IsSet())
+					{
+						NewKeyAttributesScratch.Add(FKeyAttributes().SetArriveTangent(NewTangent.GetValue()));
+					}
+					else  //Apply Delta to each Tangent indvidually
+					{
+						const FKeyAttributes Attributes = KeyData.Attributes[i];
+						const float ArriveTangent = Attributes.GetArriveTangent();
+						const FVector2D  KeyScreenPosition = FVector2D(ScreenSpace.SecondsToScreen(KeyPositions[i].InputValue), ScreenSpace.ValueToScreen(KeyPositions[i].OutputValue));
+						FVector2D TangentScreenPosition = FCurveEditor::GetVectorFromSlopeAndLength(ArriveTangent * -DisplayRatio, -60.f) + KeyScreenPosition;
+						TangentScreenPosition += PixelDelta;
+						TangentScreenPosition -= KeyScreenPosition;
+						float Tangent = (-TangentScreenPosition.Y / TangentScreenPosition.X) / DisplayRatio;
+						NewKeyAttributesScratch.Add(FKeyAttributes().SetArriveTangent(Tangent));
+					}
+				}
+			}
+			else
+			{
+				for (int32 i = 0; i < KeyData.Handles.Num(); ++i)
+				{
+					if (NewTangent.IsSet())
+					{
+						NewKeyAttributesScratch.Add(FKeyAttributes().SetLeaveTangent(NewTangent.GetValue()));
+					}
+					else  //Apply Delta to each Tangent indvidually
+					{
+						const FKeyAttributes Attributes = KeyData.Attributes[i];
+						const float LeaveTangent = Attributes.GetLeaveTangent();
+						const FVector2D  KeyScreenPosition = FVector2D(ScreenSpace.SecondsToScreen(KeyPositions[i].InputValue), ScreenSpace.ValueToScreen(KeyPositions[i].OutputValue));
+						FVector2D TangentScreenPosition = FCurveEditor::GetVectorFromSlopeAndLength(LeaveTangent * -DisplayRatio,60.f) + KeyScreenPosition;
+						TangentScreenPosition += PixelDelta;
+						TangentScreenPosition -= KeyScreenPosition;
+						float Tangent = (-TangentScreenPosition.Y / TangentScreenPosition.X) / DisplayRatio;
+						NewKeyAttributesScratch.Add(FKeyAttributes().SetLeaveTangent(Tangent));
+					}
+				}
+			}
+
+			Curve->SetKeyAttributes(KeyData.Handles, NewKeyAttributesScratch);
+		}
+	}
+}
+
+void FCurveEditorDragOperation_Tangent::OnCancelDrag()
+{
+	ICurveEditorKeyDragOperation::OnCancelDrag();
+
+	for (const FKeyData& KeyData : KeysByCurve)
+	{
+		if (FCurveModel* Curve = CurveEditor->FindCurve(KeyData.CurveID))
+		{
+			Curve->SetKeyAttributes(KeyData.Handles, KeyData.Attributes);
+		}
+	}
+}
+
+FVector2D FCurveEditorDragOperation_Tangent::RoundTrajectory(FVector2D Delta)
+{
+	float Distance = Delta.Size();
+	float Theta = FMath::Atan2(Delta.Y, Delta.X) + PI/2;
+
+	float RoundTo = PI / 4;
+	Theta = FMath::RoundToInt(Theta / RoundTo) * RoundTo - PI/2;
+	return FVector2D(Distance * FMath::Cos(Theta), Distance * FMath::Sin(Theta));
+}
\ No newline at end of file
diff --git a/Engine/Source/Editor/CurveEditor/Private/DragOperations/CurveEditorDragOperation_Tangent.h b/Engine/Source/Editor/CurveEditor/Private/DragOperations/CurveEditorDragOperation_Tangent.h
new file mode 100644
--- /dev/null
+++ b/Engine/Source/Editor/CurveEditor/Private/DragOperations/CurveEditorDragOperation_Tangent.h
@@ -0,0 +1,54 @@
+// Copyright 1998-2018 Epic Games, Inc. All Rights Reserved.
+
+#pragma once
+
+#include "CoreTypes.h"
+#include "ICurveEditorDragOperation.h"
+#include "CurveDataAbstraction.h"
+
+class FCurveEditor;
+
+class FCurveEditorDragOperation_Tangent : public ICurveEditorKeyDragOperation
+{
+public:
+
+	virtual void OnInitialize(FCurveEditor* InCurveEditor, const TOptional<FCurvePointHandle>& InCardinalPoint) override;
+	virtual void OnBeginDrag(FVector2D InitialPosition, FVector2D CurrentPosition, const FPointerEvent& MouseEvent) override;
+	virtual void OnDrag(FVector2D InitialPosition, FVector2D CurrentPosition, const FPointerEvent& MouseEvent) override;
+	virtual void OnCancelDrag() override;
+
+private:
+
+	/** Round a trajectory to the nearest 45 degrees */
+	static FVector2D RoundTrajectory(FVector2D Delta);
+
+private:
+
+	/** Cardinal screen position to compute angles from (either a key position, or the starting mouse pixel) */
+	TOptional<FVector2D> CardinalPosition;
+
+	/** The type of tangents we're dragging */
+	ECurvePointType PointType;
+
+	/** Ptr back to the curve editor */
+	FCurveEditor* CurveEditor;
+
+private:
+
+	struct FKeyData
+	{
+		FKeyData(FCurveModelID InCurveID)
+			: CurveID(InCurveID)
+		{}
+
+		/** The curve that contains the keys we're dragging */
+		FCurveModelID CurveID;
+		/** All the handles within a given curve that we are dragging */
+		TArray<FKeyHandle> Handles;
+		/** The key attributes for each of the above handles */
+		TArray<FKeyAttributes> Attributes;
+	};
+
+	/** Key dragging data stored per-curve */
+	TArray<FKeyData> KeysByCurve;
+};
\ No newline at end of file
diff --git a/Engine/Source/Editor/CurveEditor/Private/DragOperations/CurveEditorDragOperation_Zoom.cpp b/Engine/Source/Editor/CurveEditor/Private/DragOperations/CurveEditorDragOperation_Zoom.cpp
new file mode 100644
--- /dev/null
+++ b/Engine/Source/Editor/CurveEditor/Private/DragOperations/CurveEditorDragOperation_Zoom.cpp
@@ -0,0 +1,47 @@
+// Copyright 1998-2018 Epic Games, Inc. All Rights Reserved.
+
+#include "CurveEditorDragOperation_Zoom.h"
+#include "CurveEditorScreenSpace.h"
+#include "CurveEditor.h"
+
+FCurveEditorDragOperation_Zoom::FCurveEditorDragOperation_Zoom(FCurveEditor* InCurveEditor)
+	: CurveEditor(InCurveEditor)
+{
+}
+
+void FCurveEditorDragOperation_Zoom::OnBeginDrag(FVector2D InitialPosition, FVector2D CurrentPosition, const FPointerEvent& MouseEvent)
+{
+	FCurveEditorScreenSpace ScreenSpace = CurveEditor->GetScreenSpace();
+
+	ZoomFactor = InitialPosition / ScreenSpace.GetPhysicalSize();
+
+	OriginalInputRange = ( ScreenSpace.GetInputMax() - ScreenSpace.GetInputMin() );
+	OriginalOutputRange = ( ScreenSpace.GetOutputMax() - ScreenSpace.GetOutputMin() );
+
+	ZoomOriginX = ScreenSpace.GetInputMin() + OriginalInputRange * ZoomFactor.X;
+	ZoomOriginY = ScreenSpace.GetOutputMin() + OriginalOutputRange * ZoomFactor.Y;
+}
+
+void FCurveEditorDragOperation_Zoom::OnDrag(FVector2D InitialPosition, FVector2D CurrentPosition, const FPointerEvent& MouseEvent)
+{
+	FVector2D PixelDelta = GetLockedMousePosition(InitialPosition, CurrentPosition, MouseEvent) - InitialPosition;
+
+	FCurveEditorScreenSpace ScreenSpace = CurveEditor->GetScreenSpace();
+
+	double ClampRange = 1e9f;
+
+	double DiffX = double(PixelDelta.X) / (ScreenSpace.GetPhysicalWidth() / OriginalInputRange);
+	double DiffY = double(PixelDelta.Y) / (ScreenSpace.GetPhysicalHeight() / OriginalOutputRange);
+
+	double NewInputRange = OriginalInputRange + DiffX;
+	double InputMin  = FMath::Clamp<double>(ZoomOriginX - NewInputRange * ZoomFactor.X, -ClampRange, ClampRange);
+	double InputMax  = FMath::Clamp<double>(ZoomOriginX + NewInputRange * (1.f - ZoomFactor.X), InputMin, ClampRange);
+
+
+	double NewOutputRange = OriginalOutputRange + DiffY;
+	double OutputMin  = FMath::Clamp<double>(ZoomOriginY - NewOutputRange * ZoomFactor.Y, -ClampRange, ClampRange);
+	double OutputMax  = FMath::Clamp<double>(ZoomOriginY + NewOutputRange * (1.f - ZoomFactor.Y), OutputMin, ClampRange);
+
+	CurveEditor->GetBounds().SetInputBounds(InputMin, InputMax);
+	CurveEditor->GetBounds().SetOutputBounds(OutputMin, OutputMax);
+}
\ No newline at end of file
diff --git a/Engine/Source/Editor/CurveEditor/Private/DragOperations/ICurveEditorDragOperation.cpp b/Engine/Source/Editor/CurveEditor/Private/DragOperations/ICurveEditorDragOperation.cpp
new file mode 100644
--- /dev/null
+++ b/Engine/Source/Editor/CurveEditor/Private/DragOperations/ICurveEditorDragOperation.cpp
@@ -0,0 +1,60 @@
+// Copyright 1998-2018 Epic Games, Inc. All Rights Reserved.
+
+#include "ICurveEditorDragOperation.h"
+#include "CurveEditor.h"
+#include "Input/Events.h"
+
+void ICurveEditorDragOperation::BeginDrag(FVector2D InitialPosition, FVector2D CurrentPosition, const FPointerEvent& MouseEvent)
+{
+	GetLockedMousePosition(InitialPosition, CurrentPosition, MouseEvent);
+	OnBeginDrag(InitialPosition, CurrentPosition, MouseEvent);
+}
+
+void ICurveEditorDragOperation::Drag(FVector2D InitialPosition, FVector2D CurrentPosition, const FPointerEvent& MouseEvent)
+{
+	OnDrag(InitialPosition, CurrentPosition, MouseEvent);
+}
+
+void ICurveEditorDragOperation::EndDrag(FVector2D InitialPosition, FVector2D CurrentPosition, const FPointerEvent& MouseEvent)
+{
+	OnEndDrag(InitialPosition, CurrentPosition, MouseEvent);
+}
+
+int32 ICurveEditorDragOperation::Paint(const FGeometry& AllottedGeometry, FSlateWindowElementList& OutDrawElements, int32 LayerId)
+{
+	return OnPaint(AllottedGeometry, OutDrawElements, LayerId);
+}
+
+void ICurveEditorDragOperation::CancelDrag()
+{
+	OnCancelDrag();
+}
+
+FVector2D ICurveEditorDragOperation::GetLockedMousePosition(FVector2D InitialPosition, FVector2D CurrentPosition, const FPointerEvent& MouseEvent)
+{
+	if (MouseEvent.IsShiftDown())
+	{
+		if (MouseLockVector == FVector2D::UnitVector)
+		{
+			if (FMath::Abs(CurrentPosition.Y - InitialPosition.Y) <= FMath::Abs(CurrentPosition.X - InitialPosition.X))
+			{
+				MouseLockVector.Y = 0.f;
+			}
+			else
+			{
+				MouseLockVector.X = 0.f;
+			}
+		}
+	}
+	else
+	{
+		MouseLockVector = FVector2D::UnitVector;
+	}
+	return InitialPosition + (CurrentPosition - InitialPosition) * MouseLockVector;
+}
+
+void ICurveEditorKeyDragOperation::Initialize(FCurveEditor* InCurveEditor, const TOptional<FCurvePointHandle>& CardinalPoint)
+{
+	SnapMetrics = InCurveEditor->GetSnapMetrics();
+	OnInitialize(InCurveEditor, CardinalPoint);
+}
\ No newline at end of file
