Copying //UE4/Dev-Sequencer to //UE4/Dev-Main (Source: //UE4/Dev-Sequencer @ 3946692)

diff --git a/Engine/Source/Editor/CurveEditor/Private/SCurveEditorPanel.cpp b/Engine/Source/Editor/CurveEditor/Private/SCurveEditorPanel.cpp
new file mode 100644
--- /dev/null
+++ b/Engine/Source/Editor/CurveEditor/Private/SCurveEditorPanel.cpp
@@ -0,0 +1,1037 @@
+// Copyright 1998-2018 Epic Games, Inc. All Rights Reserved.
+
+#include "SCurveEditorPanel.h"
+#include "Templates/Tuple.h"
+#include "Rendering/DrawElements.h"
+#include "CurveDrawInfo.h"
+#include "CurveEditorSettings.h"
+#include "CurveEditorContextMenu.h"
+#include "DragOperations/CurveEditorDragOperation_Tangent.h"
+#include "DragOperations/CurveEditorDragOperation_MoveKeys.h"
+#include "DragOperations/CurveEditorDragOperation_Pan.h"
+#include "DragOperations/CurveEditorDragOperation_Zoom.h"
+#include "DragOperations/CurveEditorDragOperation_Marquee.h"
+#include "CurveEditorCommands.h"
+#include "Framework/Commands/UICommandList.h"
+#include "Framework/MultiBox/MultiBoxBuilder.h"
+#include "Widgets/SToolTip.h"
+#include "Widgets/SNullWidget.h"
+#include "Widgets/Layout/SScrollBox.h"
+#include "Widgets/Layout/SSplitter.h"
+#include "Widgets/Text/STextBlock.h"
+#include "EditorStyleSet.h"
+#include "Misc/Attribute.h"
+#include "Algo/Sort.h"
+#include "Widgets/Input/SNumericEntryBox.h"
+#include "Widgets/Input/STextEntryPopup.h"
+#define LOCTEXT_NAMESPACE "SCurveEditorPanel"
+
+namespace
+{
+	static float HoverProximityThresholdPx = 5.f;
+	static float HoveredCurveThickness = 5.f;
+	static float UnHoveredCurveThickness = 1.f;
+	static bool  bAntiAliasCurves = true;
+}
+
+
+TUniquePtr<ICurveEditorKeyDragOperation> CreateKeyDrag(ECurvePointType KeyType)
+{
+	switch (KeyType)
+	{
+	case ECurvePointType::ArriveTangent:
+	case ECurvePointType::LeaveTangent:
+		return MakeUnique<FCurveEditorDragOperation_Tangent>();
+
+	default:
+		return MakeUnique<FCurveEditorDragOperation_MoveKeys>();
+	}
+}
+
+class SDynamicToolTip : public SToolTip
+{
+public:
+	TAttribute<bool> bIsEnabled;
+	virtual bool IsEmpty() const override { return !bIsEnabled.Get(); }
+};
+
+void SCurveEditorPanel::Construct(const FArguments& InArgs, TSharedRef<FCurveEditor> InCurveEditor)
+{
+	GridLineTintAttribute = InArgs._GridLineTint;
+
+	CachedSelectionSerialNumber = 0;
+	CurveEditor = InCurveEditor;
+
+	ReduceTolerance = 0.1f;
+
+	CurveEditor->BindCommands();
+
+	CommandList = MakeShared<FUICommandList>();
+	CommandList->Append(InCurveEditor->GetCommands().ToSharedRef());
+
+	BindCommands();
+	SetClipping(EWidgetClipping::ClipToBounds);
+
+	ChildSlot
+	[
+		SNew(SOverlay)
+
+		+ SOverlay::Slot()
+		[
+			SNew(SSplitter)
+			.Visibility(this, &SCurveEditorPanel::GetSplitterVisibility)
+
+			+ SSplitter::Slot()
+			.Value(0.25f)
+			[
+				SAssignNew(EditPanel, SScrollBox)
+			]
+
+			+ SSplitter::Slot()
+			.Value(0.75f)
+			[
+				SNullWidget::NullWidget
+			]
+		]
+	];
+
+	UpdateEditBox();
+
+	TSharedRef<SDynamicToolTip> ToolTipWidget =
+		SNew(SDynamicToolTip)
+		.BorderImage(FCoreStyle::Get().GetBrush("ToolTip.BrightBackground"))
+		[
+			SNew(SVerticalBox)
+			+ SVerticalBox::Slot()
+			[
+				SNew(STextBlock)
+				.Text(this, &SCurveEditorPanel::GetToolTipCurveName)
+				.Font(FCoreStyle::Get().GetFontStyle("ToolTip.LargerFont"))
+				.ColorAndOpacity(FLinearColor::Black)
+			]
+			+ SVerticalBox::Slot()
+			[
+				SNew(STextBlock)
+				.Text(this, &SCurveEditorPanel::GetToolTipTimeText)
+				.Font(FCoreStyle::Get().GetFontStyle("ToolTip.LargerFont"))
+				.ColorAndOpacity(FLinearColor::Black)
+			]
+			+ SVerticalBox::Slot()
+			[
+				SNew(STextBlock)
+				.Text(this, &SCurveEditorPanel::GetToolTipValueText)
+				.Font(FCoreStyle::Get().GetFontStyle("ToolTip.LargerFont"))
+				.ColorAndOpacity(FLinearColor::Black)
+			]
+		];
+
+	ToolTipWidget->bIsEnabled = MakeAttributeSP(this, &SCurveEditorPanel::IsToolTipEnabled);
+	SetToolTip(ToolTipWidget);
+}
+
+void SCurveEditorPanel::BindCommands()
+{
+	// Interpolation and tangents
+	{
+		FExecuteAction SetConstant   = FExecuteAction::CreateSP(this, &SCurveEditorPanel::SetKeyAttributes, FKeyAttributes().SetInterpMode(RCIM_Constant).SetTangentMode(RCTM_Auto), LOCTEXT("SetInterpConstant", "Set Interp Constant"));
+		FExecuteAction SetLinear     = FExecuteAction::CreateSP(this, &SCurveEditorPanel::SetKeyAttributes, FKeyAttributes().SetInterpMode(RCIM_Linear).SetTangentMode(RCTM_Auto),   LOCTEXT("SetInterpLinear",   "Set Interp Linear"));
+		FExecuteAction SetCubicAuto  = FExecuteAction::CreateSP(this, &SCurveEditorPanel::SetKeyAttributes, FKeyAttributes().SetInterpMode(RCIM_Cubic).SetTangentMode(RCTM_Auto),    LOCTEXT("SetInterpCubic",    "Set Interp Auto"));
+		FExecuteAction SetCubicUser  = FExecuteAction::CreateSP(this, &SCurveEditorPanel::SetKeyAttributes, FKeyAttributes().SetInterpMode(RCIM_Cubic).SetTangentMode(RCTM_User),    LOCTEXT("SetInterpUser",     "Set Interp User"));
+		FExecuteAction SetCubicBreak = FExecuteAction::CreateSP(this, &SCurveEditorPanel::SetKeyAttributes, FKeyAttributes().SetInterpMode(RCIM_Cubic).SetTangentMode(RCTM_Break),   LOCTEXT("SetInterpBreak",    "Set Interp Break"));
+
+		FIsActionChecked IsConstantCommon   = FIsActionChecked::CreateSP(this, &SCurveEditorPanel::CompareCommonInterpolationMode, RCIM_Constant);
+		FIsActionChecked IsLinearCommon     = FIsActionChecked::CreateSP(this, &SCurveEditorPanel::CompareCommonInterpolationMode, RCIM_Linear);
+		FIsActionChecked IsCubicAutoCommon  = FIsActionChecked::CreateSP(this, &SCurveEditorPanel::CompareCommonTangentMode, RCIM_Cubic, RCTM_Auto);
+		FIsActionChecked IsCubicUserCommon  = FIsActionChecked::CreateSP(this, &SCurveEditorPanel::CompareCommonTangentMode, RCIM_Cubic, RCTM_User);
+		FIsActionChecked IsCubicBreakCommon = FIsActionChecked::CreateSP(this, &SCurveEditorPanel::CompareCommonTangentMode, RCIM_Cubic, RCTM_Break);
+
+		CommandList->MapAction(FCurveEditorCommands::Get().InterpolationConstant, SetConstant, FCanExecuteAction(), IsConstantCommon);
+		CommandList->MapAction(FCurveEditorCommands::Get().InterpolationLinear, SetLinear, FCanExecuteAction(), IsLinearCommon);
+		CommandList->MapAction(FCurveEditorCommands::Get().InterpolationCubicAuto, SetCubicAuto, FCanExecuteAction(), IsCubicAutoCommon);
+		CommandList->MapAction(FCurveEditorCommands::Get().InterpolationCubicUser, SetCubicUser, FCanExecuteAction(), IsCubicUserCommon);
+		CommandList->MapAction(FCurveEditorCommands::Get().InterpolationCubicBreak, SetCubicBreak, FCanExecuteAction(), IsCubicBreakCommon);
+	}
+
+	// Pre Extrapolation Modes
+	{
+		FExecuteAction SetCycle           = FExecuteAction::CreateSP(this, &SCurveEditorPanel::SetCurveAttributes, FCurveAttributes().SetPreExtrapolation(RCCE_Cycle),           LOCTEXT("SetPreExtrapCycle",           "Set Pre Extrapolation (Cycle)"));
+		FExecuteAction SetCycleWithOffset = FExecuteAction::CreateSP(this, &SCurveEditorPanel::SetCurveAttributes, FCurveAttributes().SetPreExtrapolation(RCCE_CycleWithOffset), LOCTEXT("SetPreExtrapCycleWithOffset", "Set Pre Extrapolation (Cycle With Offset)"));
+		FExecuteAction SetOscillate       = FExecuteAction::CreateSP(this, &SCurveEditorPanel::SetCurveAttributes, FCurveAttributes().SetPreExtrapolation(RCCE_Oscillate),       LOCTEXT("SetPreExtrapOscillate",       "Set Pre Extrapolation (Oscillate)"));
+		FExecuteAction SetLinear          = FExecuteAction::CreateSP(this, &SCurveEditorPanel::SetCurveAttributes, FCurveAttributes().SetPreExtrapolation(RCCE_Linear),          LOCTEXT("SetPreExtrapLinear",          "Set Pre Extrapolation (Linear)"));
+		FExecuteAction SetConstant        = FExecuteAction::CreateSP(this, &SCurveEditorPanel::SetCurveAttributes, FCurveAttributes().SetPreExtrapolation(RCCE_Constant),        LOCTEXT("SetPreExtrapConstant",        "Set Pre Extrapolation (Constant)"));
+
+		FIsActionChecked IsCycleCommon           = FIsActionChecked::CreateSP(this, &SCurveEditorPanel::CompareCommonPreExtrapolationMode, RCCE_Cycle);
+		FIsActionChecked IsCycleWithOffsetCommon = FIsActionChecked::CreateSP(this, &SCurveEditorPanel::CompareCommonPreExtrapolationMode, RCCE_CycleWithOffset);
+		FIsActionChecked IsOscillateCommon       = FIsActionChecked::CreateSP(this, &SCurveEditorPanel::CompareCommonPreExtrapolationMode, RCCE_Oscillate);
+		FIsActionChecked IsLinearCommon          = FIsActionChecked::CreateSP(this, &SCurveEditorPanel::CompareCommonPreExtrapolationMode, RCCE_Linear);
+		FIsActionChecked IsConstantCommon        = FIsActionChecked::CreateSP(this, &SCurveEditorPanel::CompareCommonPreExtrapolationMode, RCCE_Constant);
+
+		CommandList->MapAction(FCurveEditorCommands::Get().SetPreInfinityExtrapCycle, SetCycle, FCanExecuteAction(), IsCycleCommon);
+		CommandList->MapAction(FCurveEditorCommands::Get().SetPreInfinityExtrapCycleWithOffset, SetCycleWithOffset, FCanExecuteAction(), IsCycleWithOffsetCommon);
+		CommandList->MapAction(FCurveEditorCommands::Get().SetPreInfinityExtrapOscillate, SetOscillate, FCanExecuteAction(), IsOscillateCommon);
+		CommandList->MapAction(FCurveEditorCommands::Get().SetPreInfinityExtrapLinear, SetLinear, FCanExecuteAction(), IsLinearCommon);
+		CommandList->MapAction(FCurveEditorCommands::Get().SetPreInfinityExtrapConstant, SetConstant, FCanExecuteAction(), IsConstantCommon);
+	}
+
+	// Post Extrapolation Modes
+	{
+		FExecuteAction SetCycle           = FExecuteAction::CreateSP(this, &SCurveEditorPanel::SetCurveAttributes, FCurveAttributes().SetPostExtrapolation(RCCE_Cycle),           LOCTEXT("SetPostExtrapCycle",           "Set Post Extrapolation (Cycle)"));
+		FExecuteAction SetCycleWithOffset = FExecuteAction::CreateSP(this, &SCurveEditorPanel::SetCurveAttributes, FCurveAttributes().SetPostExtrapolation(RCCE_CycleWithOffset), LOCTEXT("SetPostExtrapCycleWithOffset", "Set Post Extrapolation (Cycle With Offset)"));
+		FExecuteAction SetOscillate       = FExecuteAction::CreateSP(this, &SCurveEditorPanel::SetCurveAttributes, FCurveAttributes().SetPostExtrapolation(RCCE_Oscillate),       LOCTEXT("SetPostExtrapOscillate",       "Set Post Extrapolation (Oscillate)"));
+		FExecuteAction SetLinear          = FExecuteAction::CreateSP(this, &SCurveEditorPanel::SetCurveAttributes, FCurveAttributes().SetPostExtrapolation(RCCE_Linear),          LOCTEXT("SetPostExtrapLinear",          "Set Post Extrapolation (Linear)"));
+		FExecuteAction SetConstant        = FExecuteAction::CreateSP(this, &SCurveEditorPanel::SetCurveAttributes, FCurveAttributes().SetPostExtrapolation(RCCE_Constant),        LOCTEXT("SetPostExtrapConstant",        "Set Post Extrapolation (Constant)"));
+
+		FIsActionChecked IsCycleCommon           = FIsActionChecked::CreateSP(this, &SCurveEditorPanel::CompareCommonPostExtrapolationMode, RCCE_Cycle);
+		FIsActionChecked IsCycleWithOffsetCommon = FIsActionChecked::CreateSP(this, &SCurveEditorPanel::CompareCommonPostExtrapolationMode, RCCE_CycleWithOffset);
+		FIsActionChecked IsOscillateCommon       = FIsActionChecked::CreateSP(this, &SCurveEditorPanel::CompareCommonPostExtrapolationMode, RCCE_Oscillate);
+		FIsActionChecked IsLinearCommon          = FIsActionChecked::CreateSP(this, &SCurveEditorPanel::CompareCommonPostExtrapolationMode, RCCE_Linear);
+		FIsActionChecked IsConstantCommon        = FIsActionChecked::CreateSP(this, &SCurveEditorPanel::CompareCommonPostExtrapolationMode, RCCE_Constant);
+
+		CommandList->MapAction(FCurveEditorCommands::Get().SetPostInfinityExtrapCycle, SetCycle, FCanExecuteAction(), IsCycleCommon);
+		CommandList->MapAction(FCurveEditorCommands::Get().SetPostInfinityExtrapCycleWithOffset, SetCycleWithOffset, FCanExecuteAction(), IsCycleWithOffsetCommon);
+		CommandList->MapAction(FCurveEditorCommands::Get().SetPostInfinityExtrapOscillate, SetOscillate, FCanExecuteAction(), IsOscillateCommon);
+		CommandList->MapAction(FCurveEditorCommands::Get().SetPostInfinityExtrapLinear, SetLinear, FCanExecuteAction(), IsLinearCommon);
+		CommandList->MapAction(FCurveEditorCommands::Get().SetPostInfinityExtrapConstant, SetConstant, FCanExecuteAction(), IsConstantCommon);
+	}
+	//override for key reduction so we show popup
+	{
+		CommandList->MapAction(FCurveEditorCommands::Get().ReduceCurve, FExecuteAction::CreateSP(this, &SCurveEditorPanel::OnSimplifySelection));
+	}
+}
+
+void SCurveEditorPanel::Tick(const FGeometry& AllottedGeometry, const double InCurrentTime, const float InDeltaTime)
+{
+	UpdateCommonCurveInfo();
+	UpdateEditBox();
+
+	CurveEditor->UpdateGeometry(AllottedGeometry);
+
+	CachedDrawParams.Reset();
+	CurveEditor->GetCurveDrawParams(CachedDrawParams);
+}
+
+void SCurveEditorPanel::UpdateCommonCurveInfo()
+{
+	// Gather up common extended curve info for the current set of curves
+	TOptional<FCurveAttributes> AccumulatedCurveAttributes;
+	for (const TTuple<FCurveModelID, TUniquePtr<FCurveModel>>& Pair : CurveEditor->GetCurves())
+	{
+		FCurveAttributes Attributes;
+		Pair.Value->GetCurveAttributes(Attributes);
+
+		if (!AccumulatedCurveAttributes.IsSet())
+		{
+			AccumulatedCurveAttributes = Attributes;
+		}
+		else
+		{
+			AccumulatedCurveAttributes = FCurveAttributes::MaskCommon(AccumulatedCurveAttributes.GetValue(), Attributes);
+		}
+	}
+
+	// Reset the common curve and key info
+	CachedCommonCurveAttributes = AccumulatedCurveAttributes.Get(FCurveAttributes());
+
+	TOptional<FKeyAttributes> AccumulatedKeyAttributes;
+	TArray<FKeyAttributes> AllKeyAttributes;
+
+	for (const TTuple<FCurveModelID, FKeyHandleSet>& Pair : CurveEditor->Selection.GetAll())
+	{
+		FCurveModel* Curve = CurveEditor->FindCurve(Pair.Key);
+		if (Curve)
+		{
+			AllKeyAttributes.Reset();
+			AllKeyAttributes.SetNum(Pair.Value.Num());
+
+			Curve->GetKeyAttributes(Pair.Value.AsArray(), AllKeyAttributes);
+			for (const FKeyAttributes& Attributes : AllKeyAttributes)
+			{
+				if (!AccumulatedKeyAttributes.IsSet())
+				{
+					AccumulatedKeyAttributes = Attributes;
+				}
+				else
+				{
+					AccumulatedKeyAttributes = FKeyAttributes::MaskCommon(AccumulatedKeyAttributes.GetValue(), Attributes);
+				}
+			}
+		}
+	}
+
+	// Reset the common curve and key info
+	CachedCommonKeyAttributes = AccumulatedKeyAttributes.Get(FKeyAttributes());
+}
+
+void SCurveEditorPanel::RebindContextualActions(FVector2D MousePosition)
+{
+	CommandList->UnmapAction(FCurveEditorCommands::Get().AddKeyHovered);
+	CommandList->UnmapAction(FCurveEditorCommands::Get().AddKeyToAllCurves);
+	CommandList->UnmapAction(FCurveEditorCommands::Get().AddKeyToAllCurvesHere);
+
+	TOptional<FCurveModelID> HoveredCurve = GetHoveredCurve();
+	if (HoveredCurve.IsSet())
+	{
+		CommandList->MapAction(FCurveEditorCommands::Get().AddKeyHovered, FExecuteAction::CreateSP(this, &SCurveEditorPanel::OnAddKey, HoveredCurve.GetValue(), MousePosition));
+	}
+	CommandList->MapAction(FCurveEditorCommands::Get().AddKeyToAllCurves, FExecuteAction::CreateSP(this, &SCurveEditorPanel::OnAddKey, MousePosition));
+	CommandList->MapAction(FCurveEditorCommands::Get().AddKeyToAllCurvesHere, FExecuteAction::CreateSP(this, &SCurveEditorPanel::OnAddKeyHere, MousePosition));
+}
+
+void SCurveEditorPanel::UpdateEditBox()
+{
+	const TMap<FCurveModelID, TUniquePtr<FCurveModel>>& AllCurves = CurveEditor->GetCurves();
+	for (const TTuple<FCurveModelID, TUniquePtr<FCurveModel>>& Pair : AllCurves)
+	{
+		if (!CurveToEditUI.Contains(Pair.Key))
+		{
+			TSharedPtr<SWidget> EditUI = Pair.Value->CreateEditUI(CurveEditor, Pair.Key);
+			if (EditUI.IsValid())
+			{
+				EditPanel->AddSlot()
+				.Padding(FMargin(0,0,0,10))
+				[
+					EditUI.ToSharedRef()
+				];
+			}
+			CurveToEditUI.Add(Pair.Key, EditUI);
+		}
+	}
+
+	for (auto It = CurveToEditUI.CreateIterator(); It; ++It)
+	{
+		if (!AllCurves.Contains(It.Key()))
+		{
+			if (It.Value().IsValid())
+			{
+				EditPanel->RemoveSlot(It.Value().ToSharedRef());
+			}
+
+			It.RemoveCurrent();
+		}
+	}
+}
+
+void SCurveEditorPanel::UpdateCurveProximities(FVector2D MousePixel)
+{
+	CachedToolTipData.Reset();
+
+	// Update curve proximities
+	CurveProximities.Reset();
+
+	FCurveEditorScreenSpace ScreenSpace = CurveEditor->GetScreenSpace();
+
+	double MinMouseTime    = ScreenSpace.ScreenToSeconds(MousePixel.X - HoverProximityThresholdPx);
+	double MaxMouseTime    = ScreenSpace.ScreenToSeconds(MousePixel.X + HoverProximityThresholdPx);
+	double MouseValue      = ScreenSpace.ScreenToValue(MousePixel.Y);
+	float  PixelsPerOutput = ScreenSpace.PixelsPerOutput();
+
+	FVector2D MinPos(MousePixel.X - HoverProximityThresholdPx, 0.0f);
+	FVector2D MaxPos(MousePixel.X + HoverProximityThresholdPx, 0.0f);
+
+	for (const TTuple<FCurveModelID, TUniquePtr<FCurveModel>>& Pair : CurveEditor->GetCurves())
+	{
+		double InputOffset = Pair.Value->GetInputDisplayOffset();
+		double MinEvalTime = MinMouseTime - InputOffset;
+		double MaxEvalTime = MaxMouseTime - InputOffset;
+
+		double MinValue = 0.0, MaxValue = 0.0;
+		if (Pair.Value->Evaluate(MinEvalTime, MinValue) && Pair.Value->Evaluate(MaxEvalTime, MaxValue))
+		{
+			MinPos.Y = ScreenSpace.ValueToScreen(MinValue);
+			MaxPos.Y = ScreenSpace.ValueToScreen(MaxValue);
+
+			float Distance = (FMath::ClosestPointOnSegment2D(MousePixel, MinPos, MaxPos) - MousePixel).Size();
+			CurveProximities.Add(MakeTuple(Pair.Key, Distance));
+		}
+	}
+
+	Algo::SortBy(CurveProximities, [](TTuple<FCurveModelID, float> In) { return In.Get<1>(); });
+
+	if (CurveProximities.Num() > 0 && CurveProximities[0].Get<1>() < HoverProximityThresholdPx)
+	{
+		const FCurveModel* HoveredCurve = CurveEditor->FindCurve(CurveProximities[0].Get<0>());
+
+		if (HoveredCurve)
+		{
+			double MouseTime = ScreenSpace.ScreenToSeconds(MousePixel.X) - HoveredCurve->GetInputDisplayOffset();
+			double EvaluatedTime = CurveEditor->GetSnapMetrics().SnapInputSeconds(MouseTime);
+
+			double EvaluatedValue = 0.0;
+			HoveredCurve->Evaluate(EvaluatedTime, EvaluatedValue);
+
+			FCachedToolTipData ToolTipData;
+			ToolTipData.Text           = HoveredCurve->GetDisplayName();
+			ToolTipData.EvaluatedTime  = FText::Format(NSLOCTEXT("CurveEditor", "CurveEditorTime", "{0}"), EvaluatedTime);
+			ToolTipData.EvaluatedValue = FText::Format(NSLOCTEXT("CurveEditor", "CurveEditorValue", "{0}"), EvaluatedValue);
+
+			CachedToolTipData = ToolTipData;
+		}
+	}
+}
+
+EVisibility SCurveEditorPanel::GetSplitterVisibility() const
+{
+	return EVisibility::Visible;
+}
+
+bool SCurveEditorPanel::IsToolTipEnabled() const
+{
+	return (CachedToolTipData.IsSet() && CurveEditor->GetSettings()->GetShowCurveEditorCurveToolTips());
+}
+
+FText SCurveEditorPanel::GetToolTipCurveName() const
+{
+	return CachedToolTipData.IsSet() ? CachedToolTipData->Text : FText();
+}
+
+FText SCurveEditorPanel::GetToolTipTimeText() const
+{
+	return CachedToolTipData.IsSet() ? CachedToolTipData->EvaluatedTime : FText();
+}
+
+FText SCurveEditorPanel::GetToolTipValueText() const
+{
+	return CachedToolTipData.IsSet() ? CachedToolTipData->EvaluatedValue : FText();
+}
+
+TOptional<FCurveModelID> SCurveEditorPanel::GetHoveredCurve() const
+{
+	if (CurveProximities.Num() > 0 && CurveProximities[0].Get<1>() < HoverProximityThresholdPx)
+	{
+		return CurveProximities[0].Get<0>();
+	}
+
+	return TOptional<FCurveModelID>();
+}
+
+void SCurveEditorPanel::SetKeyAttributes(FKeyAttributes KeyAttributes, FText Description)
+{
+	FScopedTransaction Transaction(Description);
+
+	for (const TTuple<FCurveModelID, FKeyHandleSet>& Pair : CurveEditor->Selection.GetAll())
+	{
+		if (FCurveModel* Curve = CurveEditor->FindCurve(Pair.Key))
+		{
+			Curve->Modify();
+			Curve->SetKeyAttributes(Pair.Value.AsArray(), KeyAttributes);
+		}
+	}
+}
+
+void SCurveEditorPanel::SetCurveAttributes(FCurveAttributes CurveAttributes, FText Description)
+{
+	FScopedTransaction Transaction(Description);
+
+	for (const TTuple<FCurveModelID, TUniquePtr<FCurveModel>>& Pair : CurveEditor->GetCurves())
+	{
+		Pair.Value->Modify();
+		Pair.Value->SetCurveAttributes(CurveAttributes);
+	}
+}
+
+void SCurveEditorPanel::OnAddKey(FVector2D MousePixel)
+{
+	FScopedTransaction Transcation(LOCTEXT("OnAddKey", "Add Key"));
+
+	bool bAddedKey = false;
+	FKeyAttributes DefaultAttributes = CurveEditor->DefaultKeyAttributes.Get();
+
+	double MouseTime = CurveEditor->GetScreenSpace().ScreenToSeconds(MousePixel.X);
+
+	CurveEditor->Selection.Clear();
+	for (const TTuple<FCurveModelID, TUniquePtr<FCurveModel>>& Pair : CurveEditor->GetCurves())
+	{
+		double EvalTime = MouseTime - Pair.Value->GetInputDisplayOffset();
+
+		double CurveValue = 0.0;
+		if (Pair.Value->Evaluate(EvalTime, CurveValue))
+		{
+			Pair.Value->Modify();
+
+			// Add a key on this curve
+			TOptional<FKeyHandle> NewKey = Pair.Value->AddKey(FKeyPosition(EvalTime, CurveValue), DefaultAttributes);
+			if (NewKey.IsSet())
+			{
+				bAddedKey = true;
+				CurveEditor->Selection.Add(FCurvePointHandle(Pair.Key, ECurvePointType::Key, NewKey.GetValue()));
+			}
+		}
+	}
+
+	if (!bAddedKey)
+	{
+		Transcation.Cancel();
+	}
+}
+
+void SCurveEditorPanel::OnAddKeyHere(FVector2D MousePixel)
+{
+	FScopedTransaction Transcation(LOCTEXT("OnAddKeyHere", "Add Key"));
+
+	bool bAddedKey = false;
+
+	FKeyAttributes DefaultAttributes = CurveEditor->DefaultKeyAttributes.Get();
+
+	double MouseTime = CurveEditor->GetScreenSpace().ScreenToSeconds(MousePixel.X);
+	double NewValue  = CurveEditor->GetScreenSpace().ScreenToValue(MousePixel.Y);
+
+	CurveEditor->Selection.Clear();
+	for (const TTuple<FCurveModelID, TUniquePtr<FCurveModel>>& Pair : CurveEditor->GetCurves())
+	{
+		Pair.Value->Modify();
+
+		double KeyTime = MouseTime - Pair.Value->GetInputDisplayOffset();
+		TOptional<FKeyHandle> NewKey = Pair.Value->AddKey(FKeyPosition(KeyTime, NewValue), DefaultAttributes);
+		if (NewKey.IsSet())
+		{
+			bAddedKey = true;
+			CurveEditor->Selection.Add(FCurvePointHandle(Pair.Key, ECurvePointType::Key, NewKey.GetValue()));
+		}
+	}
+
+	if (!bAddedKey)
+	{
+		Transcation.Cancel();
+	}
+}
+
+void SCurveEditorPanel::OnAddKey(FCurveModelID CurveToAdd, FVector2D MousePixel)
+{
+	FCurveModel* Curve = CurveEditor->FindCurve(CurveToAdd);
+	if (!Curve)
+	{
+		return;
+	}
+
+	FScopedTransaction Transcation(FText::Format(LOCTEXT("OnKeyToCurveFormat", "Add Key to Curve '{0}'"), Curve->GetDisplayName()));
+
+	FKeyAttributes DefaultAttributes = CurveEditor->DefaultKeyAttributes.Get();
+
+	double MouseTime = CurveEditor->GetScreenSpace().ScreenToSeconds(MousePixel.X);
+	double NewValue  = CurveEditor->GetScreenSpace().ScreenToValue(MousePixel.Y);
+
+	CurveEditor->Selection.Clear();
+
+	Curve->Modify();
+
+	double KeyTime = MouseTime - Curve->GetInputDisplayOffset();
+	TOptional<FKeyHandle> NewKey = Curve->AddKey(FKeyPosition(KeyTime, NewValue), DefaultAttributes);
+	if (NewKey.IsSet())
+	{
+		CurveEditor->Selection.Add(FCurvePointHandle(CurveToAdd, ECurvePointType::Key, NewKey.GetValue()));
+	}
+	else
+	{
+		Transcation.Cancel();
+	}
+}
+
+TOptional<FCurvePointHandle> SCurveEditorPanel::HitPoint(FVector2D MousePixel) const
+{
+	TOptional<FCurvePointHandle> HitPoint;
+	TOptional<float> ClosestDistance;
+
+	// Find all keys within the current hittest time
+	for (const FCurveDrawParams& Params : CachedDrawParams)
+	{
+		for (const FCurvePointInfo& Point : Params.Points)
+		{
+			const FKeyDrawInfo& PointDrawInfo = Params.GetKeyDrawInfo(Point.Type);
+
+			FSlateRect KeyRect = FSlateRect::FromPointAndExtent(Point.ScreenPosition - PointDrawInfo.ScreenSize/2, PointDrawInfo.ScreenSize);
+			if (KeyRect.ContainsPoint(MousePixel))
+			{
+				float DistanceSquared = (KeyRect.GetCenter() - MousePixel).SizeSquared();
+				if (DistanceSquared <= ClosestDistance.Get(DistanceSquared))
+				{
+					ClosestDistance = DistanceSquared;
+					HitPoint = FCurvePointHandle(Params.GetID(), Point.Type, Point.KeyHandle);
+				}
+			}
+		}
+	}
+
+	return HitPoint;
+}
+
+FReply SCurveEditorPanel::OnMouseButtonDown(const FGeometry& MyGeometry, const FPointerEvent& MouseEvent)
+{
+	CachedToolTipData.Reset();
+
+	FCurveEditorScreenSpace ScreenSpace = CurveEditor->GetScreenSpace();
+
+	FVector2D MousePixel = MyGeometry.AbsoluteToLocal(MouseEvent.GetScreenSpacePosition());
+
+	// Mouse interaction that does not require a hit test
+	if (MouseEvent.GetEffectingButton() == EKeys::MiddleMouseButton)
+	{
+		if (MouseEvent.IsAltDown())
+		{
+			DragOperation = FCurveEditorDelayedDrag(MousePixel, MouseEvent.GetEffectingButton());
+			DragOperation->DragImpl = MakeUnique<FCurveEditorDragOperation_Pan>(CurveEditor.Get());
+			return FReply::Handled().CaptureMouse(AsShared());
+		}
+		else
+		{
+			TOptional<FCurvePointHandle> NewPoint;
+
+			// Add a key to the closest curve to the mouse
+
+			if (TOptional<FCurveModelID> HoveredCurve = GetHoveredCurve())
+			{
+				FCurveModel* CurveToAddTo = CurveEditor->FindCurve(HoveredCurve.GetValue());
+				if (CurveToAddTo)
+				{
+					FScopedTransaction Transcation(LOCTEXT("InsertKey", "Insert Key"));
+
+					FKeyAttributes DefaultAttributes = CurveEditor->DefaultKeyAttributes.Get();
+
+					double MouseTime  = ScreenSpace.ScreenToSeconds(MousePixel.X);
+					double MouseValue = ScreenSpace.ScreenToValue(MousePixel.Y);
+
+					CurveToAddTo->Modify();
+
+					// Add a key on this curve
+					TOptional<FKeyHandle> NewKey = CurveToAddTo->AddKey(FKeyPosition(MouseTime, MouseValue), DefaultAttributes);
+					if (NewKey.IsSet())
+					{
+						NewPoint = FCurvePointHandle(HoveredCurve.GetValue(), ECurvePointType::Key, NewKey.GetValue());
+
+						CurveEditor->Selection.Clear();
+						CurveEditor->Selection.Add(NewPoint.GetValue());
+					}
+					else
+					{
+						Transcation.Cancel();
+					}
+				}
+			}
+
+			TUniquePtr<ICurveEditorKeyDragOperation> KeyDrag = CreateKeyDrag(CurveEditor->Selection.GetSelectionType());
+
+			bool bAllowSnapping = NewPoint.IsSet();
+
+			KeyDrag->Initialize(CurveEditor.Get(), NewPoint);
+			KeyDrag->SnapMetrics.bSnapInputValues = bAllowSnapping;
+			KeyDrag->SnapMetrics.bSnapOutputValues = bAllowSnapping;
+
+			DragOperation = FCurveEditorDelayedDrag(MousePixel, MouseEvent.GetEffectingButton());
+			DragOperation->DragImpl = MoveTemp(KeyDrag);
+
+			return FReply::Handled().CaptureMouse(AsShared());
+		}
+	}
+	else if (MouseEvent.GetEffectingButton() == EKeys::RightMouseButton)
+	{
+		if (MouseEvent.IsAltDown())
+		{
+			DragOperation = FCurveEditorDelayedDrag(MousePixel, MouseEvent.GetEffectingButton());
+			DragOperation->DragImpl = MakeUnique<FCurveEditorDragOperation_Zoom>(CurveEditor.Get());
+			return FReply::Handled().CaptureMouse(AsShared());
+		}
+	}
+
+	bool bShiftPressed = MouseEvent.IsShiftDown();
+	bool bCtrlPressed  = MouseEvent.IsControlDown();
+
+	TOptional<FCurvePointHandle> MouseDownPoint = HitPoint(MousePixel);
+	if (MouseDownPoint.IsSet())
+	{
+		if (bShiftPressed)
+		{
+			CurveEditor->Selection.Add(MouseDownPoint.GetValue());
+		}
+		else if (bCtrlPressed)
+		{
+			CurveEditor->Selection.Toggle(MouseDownPoint.GetValue());
+		}
+		else
+		{
+			if (CurveEditor->Selection.Contains(MouseDownPoint->CurveID, MouseDownPoint->KeyHandle))
+			{
+				CurveEditor->Selection.ChangeSelectionPointType(MouseDownPoint->PointType);
+			}
+			else
+			{
+				CurveEditor->Selection.Clear();
+				CurveEditor->Selection.Add(MouseDownPoint.GetValue());
+			}
+		}
+
+		if (MouseEvent.GetEffectingButton() == EKeys::LeftMouseButton)
+		{
+			TUniquePtr<ICurveEditorKeyDragOperation> KeyDrag = CreateKeyDrag(MouseDownPoint->PointType);
+
+			KeyDrag->Initialize(CurveEditor.Get(), MouseDownPoint);
+
+			DragOperation = FCurveEditorDelayedDrag(MousePixel, MouseEvent.GetEffectingButton());
+			DragOperation->DragImpl = MoveTemp(KeyDrag);
+
+			return FReply::Handled().CaptureMouse(AsShared());
+		}
+	}
+
+	if (MouseEvent.GetEffectingButton() == EKeys::LeftMouseButton)
+	{
+		DragOperation = FCurveEditorDelayedDrag(MousePixel, MouseEvent.GetEffectingButton());
+		DragOperation->DragImpl = MakeUnique<FCurveEditorDragOperation_Marquee>(CurveEditor.Get(), this);
+		return FReply::Handled().CaptureMouse(AsShared());
+	}
+	else if (MouseEvent.GetEffectingButton() == EKeys::RightMouseButton)
+	{
+		DragOperation = FCurveEditorDelayedDrag(MousePixel, MouseEvent.GetEffectingButton());
+		DragOperation->DragImpl = MakeUnique<FCurveEditorDragOperation_Pan>(CurveEditor.Get());
+
+		return FReply::Handled().CaptureMouse(AsShared());
+	}
+
+	return FReply::Handled();
+}
+
+FReply SCurveEditorPanel::OnMouseMove(const FGeometry& MyGeometry, const FPointerEvent& MouseEvent)
+{
+	FVector2D MousePixel = MyGeometry.AbsoluteToLocal(MouseEvent.GetScreenSpacePosition());
+
+	if (DragOperation.IsSet())
+	{
+		FVector2D InitialPosition = DragOperation->GetInitialPosition();
+
+		if (!DragOperation->IsDragging() && DragOperation->AttemptDragStart(MouseEvent))
+		{
+			DragOperation->DragImpl->BeginDrag(InitialPosition, MousePixel, MouseEvent);
+		}
+		else if (DragOperation->IsDragging())
+		{
+			DragOperation->DragImpl->Drag(InitialPosition, MousePixel, MouseEvent);
+		}
+	}
+	else
+	{
+		UpdateCurveProximities(MousePixel);
+	}
+
+	return FReply::Handled();
+}
+
+FReply SCurveEditorPanel::OnMouseButtonUp(const FGeometry& MyGeometry, const FPointerEvent& MouseEvent)
+{
+	FVector2D MousePosition = MyGeometry.AbsoluteToLocal(MouseEvent.GetScreenSpacePosition());
+
+	if (DragOperation.IsSet() && DragOperation->IsDragging())
+	{
+		FVector2D InitialPosition = DragOperation->GetInitialPosition();
+		DragOperation->DragImpl->EndDrag(InitialPosition, MousePosition, MouseEvent);
+	}
+	else
+	{
+		TOptional<FCurvePointHandle> MouseUpPoint = HitPoint(MousePosition);
+
+		RebindContextualActions(MousePosition);
+
+		if (MouseEvent.GetEffectingButton() == EKeys::RightMouseButton)
+		{
+			if (!MouseUpPoint.IsSet())
+			{
+				CurveEditor->Selection.Clear();
+			}
+
+			const bool bCloseAfterSelection = true;
+			FMenuBuilder MenuBuilder(bCloseAfterSelection, CommandList);
+
+			FCurveEditorContextMenu::BuildMenu(MenuBuilder, CurveEditor, MouseUpPoint, GetHoveredCurve());
+
+			// Push the context menu
+			FWidgetPath WidgetPath = MouseEvent.GetEventPath() != nullptr ? *MouseEvent.GetEventPath() : FWidgetPath();
+			FSlateApplication::Get().PushMenu(AsShared(), WidgetPath, MenuBuilder.MakeWidget(), FSlateApplication::Get().GetCursorPos(), FPopupTransitionEffect(FPopupTransitionEffect::ContextMenu));
+		}
+		else if (MouseEvent.GetEffectingButton() == EKeys::LeftMouseButton && !MouseUpPoint.IsSet())
+		{
+			if (!CommandList->ProcessCommandBindings(MouseEvent))
+			{
+				CurveEditor->Selection.Clear();
+			}
+		}
+	}
+
+	UpdateCurveProximities(MousePosition);
+
+	DragOperation.Reset();
+	return FReply::Handled().ReleaseMouseCapture();
+}
+
+FReply SCurveEditorPanel::OnMouseWheel(const FGeometry& MyGeometry, const FPointerEvent& MouseEvent)
+{
+	FCurveEditorScreenSpace ScreenSpace = CurveEditor->GetScreenSpace();
+
+	FVector2D MousePixel   = MyGeometry.AbsoluteToLocal(MouseEvent.GetScreenSpacePosition());
+	double    CurrentTime  = ScreenSpace.ScreenToSeconds(MousePixel.X);
+	double    CurrentValue = ScreenSpace.ScreenToValue(MousePixel.Y);
+
+	float ZoomDelta = 1.f - FMath::Clamp(0.1f * MouseEvent.GetWheelDelta(), -0.9f, 0.9f);
+	CurveEditor->ZoomAround(ZoomDelta, CurrentTime, CurrentValue);
+
+	return FReply::Handled();
+}
+
+void SCurveEditorPanel::OnMouseEnter(const FGeometry& MyGeometry, const FPointerEvent& MouseEvent)
+{
+	CurveProximities.Reset();
+
+	SCompoundWidget::OnMouseEnter(MyGeometry, MouseEvent);
+}
+
+void SCurveEditorPanel::OnMouseLeave(const FPointerEvent& MouseEvent)
+{
+	CurveProximities.Reset();
+
+	SCompoundWidget::OnMouseLeave(MouseEvent);
+}
+
+void SCurveEditorPanel::OnFocusLost(const FFocusEvent& InFocusEvent)
+{
+	if (DragOperation.IsSet())
+	{
+		DragOperation->DragImpl->CancelDrag();
+		DragOperation.Reset();
+	}
+
+	SCompoundWidget::OnFocusLost(InFocusEvent);
+}
+
+FReply SCurveEditorPanel::OnKeyDown(const FGeometry& MyGeometry, const FKeyEvent& InKeyEvent)
+{
+	if (InKeyEvent.GetKey() == EKeys::Escape)
+	{
+		if (DragOperation.IsSet())
+		{
+			DragOperation->DragImpl->CancelDrag();
+			DragOperation.Reset();
+		}
+		else
+		{
+			CurveEditor->Selection.Clear();
+		}
+		return FReply::Handled();
+	}
+	else if (CommandList->ProcessCommandBindings(InKeyEvent))
+	{
+		return FReply::Handled();
+	}
+	return FReply::Unhandled();
+}
+
+int32 SCurveEditorPanel::OnPaint(const FPaintArgs& Args, const FGeometry& AllottedGeometry, const FSlateRect& MyCullingRect, FSlateWindowElementList& OutDrawElements, int32 LayerId, const FWidgetStyle& InWidgetStyle, bool bParentEnabled) const
+{
+	const ESlateDrawEffect DrawEffects = ShouldBeEnabled(bParentEnabled) ? ESlateDrawEffect::None : ESlateDrawEffect::DisabledEffect;
+
+	FSlateDrawElement::MakeBox(OutDrawElements,	LayerId, AllottedGeometry.ToPaintGeometry(),
+		FEditorStyle::GetBrush("ToolPanel.GroupBorder"), DrawEffects);
+
+	LayerId = DrawGridLines(AllottedGeometry, OutDrawElements, LayerId, DrawEffects);
+	LayerId = DrawCurves(AllottedGeometry, OutDrawElements, LayerId, InWidgetStyle, DrawEffects);
+
+	if (DragOperation.IsSet() && DragOperation->IsDragging())
+	{
+		LayerId = DragOperation->DragImpl->Paint(AllottedGeometry, OutDrawElements, LayerId);
+	}
+	return SCompoundWidget::OnPaint(Args, AllottedGeometry, MyCullingRect, OutDrawElements, LayerId, InWidgetStyle, bParentEnabled);
+}
+
+int32 SCurveEditorPanel::DrawGridLines(const FGeometry& AllottedGeometry, FSlateWindowElementList& OutDrawElements, int32 LayerId, ESlateDrawEffect DrawEffects) const
+{
+	// Rendering info
+	const float          Width          = AllottedGeometry.GetLocalSize().X;
+	const float          Height         = AllottedGeometry.GetLocalSize().Y;
+	const FLinearColor   MajorGridColor = GridLineTintAttribute.Get();
+	const FLinearColor   MinorGridColor = MajorGridColor.CopyWithNewOpacity(MajorGridColor.A * .25f);
+	const FPaintGeometry PaintGeometry  = AllottedGeometry.ToPaintGeometry();
+
+	TArray<float> MajorGridLines, MinorGridLines;
+	CurveEditor->GetGridLinesX(MajorGridLines, MinorGridLines);
+
+	TArray<FVector2D> LinePoints;
+	LinePoints.Add(FVector2D(0.f, 0.f));
+	LinePoints.Add(FVector2D(0.f, Height));
+
+	// Draw major grid lines
+	for (float PosX : MajorGridLines)
+	{
+		LinePoints[0].X = LinePoints[1].X = PosX;
+
+		FSlateDrawElement::MakeLines(
+			OutDrawElements,
+			LayerId,
+			PaintGeometry,
+			LinePoints,
+			DrawEffects,
+			MajorGridColor,
+			false
+		);
+	}
+
+	for (float PosX : MinorGridLines)
+	{
+		LinePoints[0].X = LinePoints[1].X = PosX;
+
+		FSlateDrawElement::MakeLines(
+			OutDrawElements,
+			LayerId,
+			PaintGeometry,
+			LinePoints,
+			DrawEffects,
+			MinorGridColor,
+			false
+		);
+	}
+
+	MajorGridLines.Reset();
+	MinorGridLines.Reset();
+	CurveEditor->GetGridLinesY(MajorGridLines, MinorGridLines);
+
+	LinePoints[0].X = 0.f;
+	LinePoints[1].X = Width;
+	for (float PosY : MajorGridLines)
+	{
+		LinePoints[0].Y = LinePoints[1].Y = PosY;
+
+		FSlateDrawElement::MakeLines(
+			OutDrawElements,
+			LayerId,
+			PaintGeometry,
+			LinePoints,
+			DrawEffects,
+			MajorGridColor,
+			false
+		);
+	}
+
+	for (float PosX : MinorGridLines)
+	{
+		LinePoints[0].X = LinePoints[1].X = PosX;
+
+		FSlateDrawElement::MakeLines(
+			OutDrawElements,
+			LayerId,
+			PaintGeometry,
+			LinePoints,
+			DrawEffects,
+			MinorGridColor,
+			false
+		);
+	}
+
+	return LayerId + 1;
+}
+
+int32 SCurveEditorPanel::DrawCurves(const FGeometry& AllottedGeometry, FSlateWindowElementList& OutDrawElements, int32 LayerId, const FWidgetStyle& InWidgetStyle, ESlateDrawEffect DrawEffects) const
+{
+	FCurveEditorScreenSpace ScreenSpace = CurveEditor->GetScreenSpace();
+
+	static const FName SelectionColorName("SelectionColor");
+	FLinearColor SelectionColor = FEditorStyle::GetSlateColor(SelectionColorName).GetColor(InWidgetStyle);
+
+	const FVector2D      VisibleSize   = AllottedGeometry.GetLocalSize();
+	const FPaintGeometry PaintGeometry = AllottedGeometry.ToPaintGeometry();
+
+	TOptional<FCurveModelID> HoveredCurve = GetHoveredCurve();
+	for (const FCurveDrawParams& Params : CachedDrawParams)
+	{
+		float Thickness = ( HoveredCurve.IsSet() && HoveredCurve.GetValue() == Params.GetID() ) ? HoveredCurveThickness : UnHoveredCurveThickness;
+		FSlateDrawElement::MakeLines(
+			OutDrawElements,
+			LayerId,
+			PaintGeometry,
+			Params.InterpolatingPoints,
+			DrawEffects,
+			Params.Color,
+			bAntiAliasCurves,
+			Thickness
+		);
+
+		TArray<FVector2D> LinePoints;
+		LinePoints.SetNum(2);
+
+		// Draw tangents
+		for (const FCurvePointInfo& Point : Params.Points)
+		{
+			const FKeyDrawInfo& PointDrawInfo = Params.GetKeyDrawInfo(Point.Type);
+			const bool          bSelected     = CurveEditor->Selection.IsSelected(FCurvePointHandle(Params.GetID(), Point.Type, Point.KeyHandle));
+			const FLinearColor  PointTint     = bSelected ? SelectionColor : PointDrawInfo.Tint;
+
+			if (Point.LineDelta.X != 0.f || Point.LineDelta.Y != 0.f)
+			{
+				LinePoints[0] = Point.ScreenPosition;
+				LinePoints[1] = Point.ScreenPosition + Point.LineDelta;
+
+				// Draw the connecting line
+				FSlateDrawElement::MakeLines(OutDrawElements, LayerId, PaintGeometry, LinePoints, DrawEffects, PointTint, true);
+			}
+
+			FPaintGeometry PointGeometry = AllottedGeometry.ToPaintGeometry(
+				Point.ScreenPosition - (PointDrawInfo.ScreenSize * 0.5f),
+				PointDrawInfo.ScreenSize
+			);
+
+			FSlateDrawElement::MakeBox(OutDrawElements, LayerId+Point.LayerBias, PointGeometry, PointDrawInfo.Brush, DrawEffects, PointTint);
+		}
+	}
+
+	return LayerId + 1;
+}
+
+
+void SCurveEditorPanel::OnSimplifySelection()
+{
+	// Display dialog and let user enter tolerance.
+	GenericTextEntryModeless(
+		NSLOCTEXT("CurveEditor.Popups", "ReduceCurveTolerance", "Tolerance"),
+		FText::AsNumber(ReduceTolerance),
+		FOnTextCommitted::CreateSP(this, &SCurveEditorPanel::OnSimplifySelectionCommited)
+	);
+}
+
+void SCurveEditorPanel::OnSimplifySelectionCommited(const FText& InText, ETextCommit::Type CommitInfo)
+{
+	CloseEntryPopupMenu();
+	if (CommitInfo == ETextCommit::OnEnter)
+	{
+		double NewTolerance = FCString::Atod(*InText.ToString());
+		const bool bIsNumber = InText.IsNumeric();
+		if (!bIsNumber)
+			return;
+		CurveEditor->SimplifySelection((float)NewTolerance);
+	}
+}
+
+
+
+void SCurveEditorPanel::GenericTextEntryModeless(const FText& DialogText, const FText& DefaultText, FOnTextCommitted OnTextComitted)
+{
+	TSharedRef<STextEntryPopup> TextEntryPopup =
+		SNew(STextEntryPopup)
+		.Label(DialogText)
+		.DefaultText(DefaultText)
+		.OnTextCommitted(OnTextComitted)
+		.ClearKeyboardFocusOnCommit(false)
+		.SelectAllTextWhenFocused(true)
+		.MaxWidth(1024.0f);
+
+	EntryPopupMenu = FSlateApplication::Get().PushMenu(
+		SharedThis(this),
+		FWidgetPath(),
+		TextEntryPopup,
+		FSlateApplication::Get().GetCursorPos(),
+		FPopupTransitionEffect(FPopupTransitionEffect::TypeInPopup)
+	);
+}
+
+
+void SCurveEditorPanel::CloseEntryPopupMenu()
+{
+	if (EntryPopupMenu.IsValid())
+	{
+		EntryPopupMenu.Pin()->Dismiss();
+	}
+}
+
+#undef LOCTEXT_NAMESPACE
\ No newline at end of file
diff --git a/Engine/Source/Editor/CurveEditor/Public/CurveDataAbstraction.h b/Engine/Source/Editor/CurveEditor/Public/CurveDataAbstraction.h
new file mode 100644
--- /dev/null
+++ b/Engine/Source/Editor/CurveEditor/Public/CurveDataAbstraction.h
@@ -0,0 +1,191 @@
+// Copyright 1998-2018 Epic Games, Inc. All Rights Reserved.
+
+#pragma once
+
+#include "CoreTypes.h"
+#include "CurveEditorTypes.h"
+#include "Curves/RichCurve.h"
+
+struct FSlateBrush;
+
+/**
+ * Generic key position information for a key on a curve
+ */
+struct FKeyPosition
+{
+	FKeyPosition()
+		: InputValue(0.0), OutputValue(0.0)
+	{}
+
+	FKeyPosition(double Input, double Output)
+		: InputValue(Input), OutputValue(Output)
+	{}
+
+	/** The key's input (x-axis) position (i.e. it's time) */
+	double InputValue;
+	/** The key's output (t-axis) position (i.e. it's value) */
+	double OutputValue;
+};
+
+/**
+ * Extended attributes that the curve editor understands
+ */
+struct FKeyAttributes
+{
+	FKeyAttributes()
+	{
+		bHasArriveTangent = 0;
+		bHasLeaveTangent  = 0;
+		bHasInterpMode    = 0;
+		bHasTangentMode   = 0;
+	}
+
+	/**
+	 * Check whether this key has the specified attributes
+	 */
+	bool HasArriveTangent() const { return bHasArriveTangent; }
+	bool HasLeaveTangent() const  { return bHasLeaveTangent;  }
+	bool HasInterpMode() const    { return bHasInterpMode;    }
+	bool HasTangentMode() const   { return bHasTangentMode;   }
+
+	/**
+	 * Retrieve specific attributes for this key. Must check such attributes exist.
+	 */
+	float GetArriveTangent() const               { check(bHasArriveTangent); return ArriveTangent; }
+	float GetLeaveTangent() const                { check(bHasLeaveTangent);  return LeaveTangent;  }
+	ERichCurveInterpMode GetInterpMode() const   { check(bHasInterpMode);    return InterpMode;    }
+	ERichCurveTangentMode GetTangentMode() const { check(bHasTangentMode);   return TangentMode;   }
+
+	/**
+	 * Set the attributes for this key
+	 */
+	FKeyAttributes& SetArriveTangent(float InArriveTangent)             { bHasArriveTangent = 1;    ArriveTangent = InArriveTangent; return *this; }
+	FKeyAttributes& SetLeaveTangent(float InLeaveTangent)               { bHasLeaveTangent = 1;     LeaveTangent = InLeaveTangent;   return *this; }
+	FKeyAttributes& SetInterpMode(ERichCurveInterpMode InInterpMode)    { bHasInterpMode = 1;       InterpMode = InInterpMode;       return *this; }
+	FKeyAttributes& SetTangentMode(ERichCurveTangentMode InTangentMode) { bHasTangentMode = 1;      TangentMode = InTangentMode;     return *this; }
+
+	/**
+	 * Reset specific attributes of this key, implying such attributes are not supported
+	 */
+	void UnsetArriveTangent() { bHasArriveTangent = 0; }
+	void UnsetLeaveTangent()  { bHasLeaveTangent = 0;  }
+	void UnsetInterpMode()    { bHasInterpMode = 0;    }
+	void UnsetTangentMode()   { bHasTangentMode = 0;   }
+
+	/**
+	 * Generate a new set of attributes that contains only those attributes common to both A and B
+	 */
+	static FKeyAttributes MaskCommon(const FKeyAttributes& A, const FKeyAttributes& B)
+	{
+		FKeyAttributes NewAttributes;
+		if (A.bHasArriveTangent && B.bHasArriveTangent && A.ArriveTangent == B.ArriveTangent)
+		{
+			NewAttributes.SetArriveTangent(A.ArriveTangent);
+		}
+
+		if (A.bHasLeaveTangent && B.bHasLeaveTangent && A.LeaveTangent == B.LeaveTangent)
+		{
+			NewAttributes.SetLeaveTangent(A.LeaveTangent);
+		}
+
+		if (A.bHasInterpMode && B.bHasInterpMode && A.InterpMode == B.InterpMode)
+		{
+			NewAttributes.SetInterpMode(A.InterpMode);
+		}
+
+		if (A.bHasTangentMode && B.bHasTangentMode && A.TangentMode == B.TangentMode)
+		{
+			NewAttributes.SetTangentMode(A.TangentMode);
+		}
+
+		return NewAttributes;
+	}
+
+private:
+
+	/** True if this key supports entry tangents */
+	uint8 bHasArriveTangent : 1;
+	/** True if this key supports exit tangents */
+	uint8 bHasLeaveTangent : 1;
+	/** True if this key supports interpolation modes */
+	uint8 bHasInterpMode : 1;
+	/** True if this key supports tangent modes */
+	uint8 bHasTangentMode : 1;
+
+	/** This key's entry tangent, if bHasArriveTangent */
+	float ArriveTangent;
+	/** This key's exit tangent, if bHasLeaveTangent */
+	float LeaveTangent;
+	/** This key's interpolation mode, if bHasInterpMode */
+	ERichCurveInterpMode InterpMode;
+	/** This key's tangent mode, if bHasTangentMode */
+	ERichCurveTangentMode TangentMode;
+};
+
+/**
+ * Structure allowing external curve data to express extended attributes
+ */
+struct FCurveAttributes
+{
+	FCurveAttributes()
+	{
+		bHasPreExtrapolation  = 0;
+		bHasPostExtrapolation = 0;
+	}
+
+	/**
+	 * Check whether this curve has the specified properties
+	 */
+	bool HasPreExtrapolation() const  { return bHasPreExtrapolation;  }
+	bool HasPostExtrapolation() const { return bHasPostExtrapolation; }
+
+	/**
+	 * Access the extended properties of this curve. Must check whether the curve has such properties first
+	 */
+	ERichCurveExtrapolation GetPreExtrapolation() const  { check(bHasPreExtrapolation);  return PreExtrapolation;  }
+	ERichCurveExtrapolation GetPostExtrapolation() const { check(bHasPostExtrapolation); return PostExtrapolation; }
+
+	/**
+	 * Set the extended properties of this curve
+	 */
+	FCurveAttributes& SetPreExtrapolation(ERichCurveExtrapolation InPreExtrapolation)   { bHasPreExtrapolation = 1;  PreExtrapolation = InPreExtrapolation;   return *this; }
+	FCurveAttributes& SetPostExtrapolation(ERichCurveExtrapolation InPostExtrapolation) { bHasPostExtrapolation = 1; PostExtrapolation = InPostExtrapolation; return *this; }
+
+	/**
+	 * Reset the extended properties of this curve, implying it does not support such properties
+	 */
+	void UnsetPreExtrapolation()  { bHasPreExtrapolation = 0;  }
+	void UnsetPostExtrapolation() { bHasPostExtrapolation = 0; }
+
+	/**
+	 * Generate a new set of attributes that contains only those attributes common to both A and B
+	 */
+	static FCurveAttributes MaskCommon(const FCurveAttributes& A, const FCurveAttributes& B)
+	{
+		FCurveAttributes NewAttributes;
+
+		if (A.bHasPreExtrapolation && B.bHasPreExtrapolation && A.PreExtrapolation == B.PreExtrapolation)
+		{
+			NewAttributes.SetPreExtrapolation(A.PreExtrapolation);
+		}
+
+		if (A.bHasPostExtrapolation && B.bHasPostExtrapolation && A.PostExtrapolation == B.PostExtrapolation)
+		{
+			NewAttributes.SetPostExtrapolation(A.PostExtrapolation);
+		}
+
+		return NewAttributes;
+	}
+
+private:
+
+	/** true if the curve can express pre-extrapolation modes */
+	uint8 bHasPreExtrapolation : 1;
+	/** true if the curve can express post-extrapolation modes */
+	uint8 bHasPostExtrapolation : 1;
+
+	/** This curve's pre-extrapolation mode. Only valid to read if bHasPreExtrapolation is true */
+	ERichCurveExtrapolation PreExtrapolation;
+	/** This curve's post-extrapolation mode. Only valid to read if bHasPostExtrapolation is true */
+	ERichCurveExtrapolation PostExtrapolation;
+};
\ No newline at end of file
diff --git a/Engine/Source/Editor/CurveEditor/Public/CurveEditor.h b/Engine/Source/Editor/CurveEditor/Public/CurveEditor.h
new file mode 100644
--- /dev/null
+++ b/Engine/Source/Editor/CurveEditor/Public/CurveEditor.h
@@ -0,0 +1,308 @@
+// Copyright 1998-2018 Epic Games, Inc. All Rights Reserved.
+
+#pragma once
+
+#include "CoreTypes.h"
+#include "Containers/Map.h"
+#include "Templates/UniquePtr.h"
+#include "Math/Vector2D.h"
+#include "Misc/Attribute.h"
+#include "Math/Range.h"
+#include "FrameRate.h"
+
+#include "CurveEditorTypes.h"
+#include "CurveDataAbstraction.h"
+#include "CurveModel.h"
+#include "ICurveEditorBounds.h"
+#include "CurveEditorSelection.h"
+#include "ICurveEditorDragOperation.h"
+
+class UCurveEditorSettings;
+class FUICommandList;
+struct FGeometry;
+struct FCurveEditorSnapMetrics;
+struct FCurveEditorScreenSpace;
+
+DECLARE_DELEGATE_OneParam(FOnSetBoolean, bool)
+
+class CURVEEDITOR_API FCurveEditor : public TSharedFromThis<FCurveEditor>
+{
+public:
+
+	/**
+	 * Container holding the current key/tangent selection
+	 */
+	FCurveEditorSelection Selection;
+
+public:
+
+	/** Attribute used to retrieve the current input snap rate (also used for display) */
+	TAttribute<FFrameRate> InputSnapRateAttribute;
+
+	/** Attribute used to retrieve the current output snap interval */
+	TAttribute<double> OutputSnapIntervalAttribute;
+
+	/** Attribute used to determine if we should snap input values */
+	TAttribute<bool> InputSnapEnabledAttribute;
+
+	/** Attribute used to determine if we should snap output values */
+	TAttribute<bool> OutputSnapEnabledAttribute;
+
+	/** Delegate that is invoked when the input snapping has been enabled/disabled */
+	FOnSetBoolean OnInputSnapEnabledChanged;
+
+	/** Delegate that is invoked when the output snapping has been enabled/disabled */
+	FOnSetBoolean OnOutputSnapEnabledChanged;
+
+	/** Attribute used for determining default attributes to apply to a newly create key */
+	TAttribute<FKeyAttributes> DefaultKeyAttributes;
+
+public:
+
+	/**
+	 * Constructor
+	 */
+	FCurveEditor();
+
+	/**
+	 * Non-copyable (shared ptr semantics)
+	 */
+	FCurveEditor(const FCurveEditor&) = delete;
+	FCurveEditor& operator=(const FCurveEditor&) = delete;
+
+	virtual ~FCurveEditor() {}
+
+public:
+
+	/**
+	 * Generate a utility struct for converting between screen (slate unit) space and the underlying input/output axes
+	 */
+	FCurveEditorScreenSpace GetScreenSpace() const;
+
+	/**
+	 * Generate a utility struct for snapping values
+	 */
+	FCurveEditorSnapMetrics GetSnapMetrics() const;
+
+public:
+
+	/**
+	 * Find a curve by its ID
+	 *
+	 * @return a ptr to the curve if found, nullptr otherwise
+	 */
+	FCurveModel* FindCurve(FCurveModelID CurveID) const;
+
+	/**
+	 * Access all the curves currently being shown on this editor
+	 */
+	const TMap<FCurveModelID, TUniquePtr<FCurveModel>>& GetCurves() const;
+
+	/**
+	 * Add a new curve to this editor
+	 */
+	FCurveModelID AddCurve(TUniquePtr<FCurveModel>&& InCurve);
+
+	/**
+	 * Remove a curve from this editor
+	 */
+	void RemoveCurve(FCurveModelID InCurveID);
+
+	/**
+	 * Retrieve this curve editor's command list
+	 */
+	FORCEINLINE TSharedPtr<FUICommandList> GetCommands() const { return CommandList; }
+
+	/**
+	 * Retrieve this curve editor's settings
+	 */
+	FORCEINLINE UCurveEditorSettings* GetSettings() const { return Settings; }
+
+public:
+
+	/**
+	 * Assign a new bounds container to this curve editor
+	 */
+	void SetBounds(TUniquePtr<ICurveEditorBounds>&& InBounds);
+
+	/**
+	 * Retrieve the current curve editor bounds implementation
+	 */
+	const ICurveEditorBounds& GetBounds() const
+	{
+		return *Bounds.Get();
+	}
+
+	/**
+	 * Retrieve the current curve editor bounds implementation
+	 */
+	ICurveEditorBounds& GetBounds()
+	{
+		return *Bounds.Get();
+	}
+
+public:
+
+	/**
+	 * Check whether this curve editor can automatically zoom to the current selection
+	 */
+	bool ShouldAutoFrame() const;
+
+	/**
+	 * Zoom the curve editor in or out around the center point
+	 *
+	 * @param Amount       The amount to zoom by as a factor of the current size
+	 */
+	void Zoom(float Amount);
+
+	/**
+	 * Zoom the curve editor in or out around the specified point
+	 *
+	 * @param Amount       The amount to zoom by as a factor of the current size
+	 * @param TimeOrigin   The time origin to zoom around
+	 * @param ValueOrigin  The value origin to zoom around
+	 */
+	void ZoomAround(float Amount, double TimeOrigin, double ValueOrigin);
+
+	/**
+	 * Zoom the curve editor to fit all the currently visible curves
+	 *
+	 * @param Axes         (Optional) Axes to lock the zoom to
+	 */
+	void ZoomToFit(EAxisList::Type Axes = EAxisList::All);
+
+	/**
+	 * Zoom the curve editor to fit the requested curves.
+	 *
+	 * @param CurveModelIDs The curve IDs to zoom to fit.
+	 * @param Axes         (Optional) Axes to lock the zoom to
+	 */
+	void ZoomToFitCurves(TArrayView<const FCurveModelID> CurveModelIDs, EAxisList::Type Axes = EAxisList::All);
+
+	/**
+	 * Zoom the curve editor to fit all the current key selection. Zooms to fit all if less than 2 keys are selected.
+	 *
+	 * @param Axes         (Optional) Axes to lock the zoom to
+	 */
+	void ZoomToFitSelection(EAxisList::Type Axes = EAxisList::All);
+
+public:
+
+	/**
+	 * Check whether keys should be snapped to the input display rate when dragging around
+	 */
+	bool IsInputSnappingEnabled() const;
+	void ToggleInputSnapping();
+
+	/**
+	 * Check whether keys should be snapped to the output snap interval when dragging around
+	 */
+	bool IsOutputSnappingEnabled() const;
+	void ToggleOutputSnapping();
+
+public:
+
+	/**
+	 * Delete the currently selected keys
+	 */
+	void DeleteSelection();
+
+	/**
+	 * Flatten the tangents on the selected keys
+	 */
+	void FlattenSelection();
+
+	/**
+	 * Straighten the tangents on the selected keys
+	 */
+	void StraightenSelection();
+
+	/**
+	 * Bake curves between selected keys by adding points at every frame of the display rate
+	 */
+	void BakeSelection();
+
+	/**
+	 * Simplify curves between the selected keys by removing redundant keys
+	 * @param Tolerance   Threshold at which to remove keys
+	 */
+	void SimplifySelection(float Tolerance= 0.1);
+
+public:
+
+	/**
+	 * Populate the specified array with curve painting parameters
+	 *
+	 * @param OutDrawParams    An array to populate with curve painting parameters, one per visible curve
+	 */
+	void GetCurveDrawParams(TArray<FCurveDrawParams>& OutDrawParams) const;
+
+	/**
+	 * Called by SCurveEditorPanel to update the allocated geometry for this curve editor.
+	 */
+	void UpdateGeometry(const FGeometry& CurrentGeometry);
+
+public:
+
+	/**
+	 * Called by SCurveEditorPanel to determine where to draw grid lines along the X-axis
+	 */
+	virtual void GetGridLinesX(TArray<float>& MajorGridLines, TArray<float>& MinorGridLines) const
+	{
+		ConstructInputGridLines(MajorGridLines, MinorGridLines);
+	}
+
+	/**
+	 * Called by SCurveEditorPanel to determine where to draw grid lines along the Y-axis
+	 */
+	virtual void GetGridLinesY(TArray<float>& MajorGridLines, TArray<float>& MinorGridLines) const
+	{
+		ConstructOutputGridLines(MajorGridLines, MinorGridLines, 4);
+	}
+
+	/**
+	 * Bind UI commands that this curve editor responds to
+	 */
+	void BindCommands();
+
+public:
+	/**
+	* Get A Vector for the given slope, usually a tangent, and length. Used to draw the tangent.
+	*/
+	static FVector2D GetVectorFromSlopeAndLength(float Slope, float Length);
+protected:
+
+	/**
+	 * Construct grid lines along the current display frame rate or time-base
+	 */
+	void ConstructInputGridLines(TArray<float>& MajorGridLines, TArray<float>& MinorGridLines) const;
+
+	/**
+	 * Construct grid lines for the current visible value range
+	 */
+	void ConstructOutputGridLines(TArray<float>& MajorGridLines, TArray<float>& MinorGridLines, uint8 MinorDivisions) const;
+
+	/**
+	 * Internal zoom to fit implementation
+	 */
+	void ZoomToFitInternal(EAxisList::Type Axes, double InputMin, double InputMax, double OutputMin, double OutputMax);
+
+protected:
+
+	/** Curve editor bounds implementation */
+	TUniquePtr<ICurveEditorBounds> Bounds;
+
+	/** Map of all the currently visible curve models */
+	TMap<FCurveModelID, TUniquePtr<FCurveModel>> CurveData;
+
+	/** UI command list of actions mapped to this curve editor */
+	TSharedPtr<FUICommandList> CommandList;
+
+	/** Curve editor settings object */
+	UCurveEditorSettings* Settings;
+
+private:
+
+	/** Cached physical size of the panel representing this editor */
+	FVector2D CachedPhysicalSize;
+};
diff --git a/Engine/Source/Editor/CurveEditor/Public/CurveEditorScreenSpace.h b/Engine/Source/Editor/CurveEditor/Public/CurveEditorScreenSpace.h
new file mode 100644
--- /dev/null
+++ b/Engine/Source/Editor/CurveEditor/Public/CurveEditorScreenSpace.h
@@ -0,0 +1,90 @@
+// Copyright 1998-2018 Epic Games, Inc. All Rights Reserved.
+
+#pragma once
+
+#include "CoreTypes.h"
+#include "FrameRate.h"
+
+template<typename> class TRange;
+struct FGeometry;
+
+/**
+ * Utility struct used for converting to/from curve editor screen space
+ */
+struct FCurveEditorScreenSpace
+{
+	/**
+	 * Construction from a physical size, and input/output range
+	 */
+	FCurveEditorScreenSpace(FVector2D InPixelSize, float InInputMin, float InInputMax, float InOutputMin, float InOutputMax)
+		: PixelSize(InPixelSize)
+		, InputMin(InInputMin), InputMax(InInputMax)
+		, OutputMin(InOutputMin), OutputMax(InOutputMax)
+	{
+	}
+
+public:
+
+	/** Convert a horizontal screen position in slate units to a value in seconds */
+	FORCEINLINE double ScreenToSeconds(float ScreenPosition) const
+	{
+		return InputMin + ScreenPosition / PixelsPerInput();
+	}
+
+	/** Convert a value in seconds to a horizontal screen position in slate units */
+	FORCEINLINE float SecondsToScreen(double InSeconds) const
+	{
+		return (InSeconds - InputMin) * PixelsPerInput();
+	}
+
+	/** Convert a vertical screen position in slate units to a value */
+	FORCEINLINE double ScreenToValue(float ScreenPosition) const
+	{
+		return OutputMin + (PixelSize.Y - ScreenPosition) / PixelsPerOutput();
+	}
+
+	/** Convert a value to a vertical screen position in slate units */
+	FORCEINLINE float ValueToScreen(double InValue) const
+	{
+		return PixelSize.Y - (InValue - OutputMin) * PixelsPerOutput();
+	}
+
+public:
+
+	/** Retrieve the number of slate units per input value */
+	FORCEINLINE float PixelsPerInput() const
+	{
+		float InputDiff = FMath::Max(InputMax - InputMin, 1e-10);
+		return PixelSize.X / InputDiff;
+	}
+
+	/** Retrieve the number of slate units per output value */
+	FORCEINLINE float PixelsPerOutput() const
+	{
+		float OutputDiff = FMath::Max(OutputMax - OutputMin, 1e-10);
+		return PixelSize.Y / OutputDiff;
+	}
+
+public:
+
+	/** Get the minimum input value displayed on the screen */
+	FORCEINLINE float GetInputMin() const { return InputMin; }
+	/** Get the maximum input value displayed on the screen */
+	FORCEINLINE float GetInputMax() const { return InputMax; }
+	/** Get the minimum output value displayed on the screen */
+	FORCEINLINE float GetOutputMin() const { return OutputMin; }
+	/** Get the maximum output value displayed on the screen */
+	FORCEINLINE float GetOutputMax() const { return OutputMax; }
+	/** Get the physical size of the screen */
+	FORCEINLINE FVector2D GetPhysicalSize() const { return PixelSize; }
+	/** Get the physical width of the screen */
+	FORCEINLINE float GetPhysicalWidth() const { return PixelSize.X; }
+	/** Get the physical height of the screen */
+	FORCEINLINE float GetPhysicalHeight() const { return PixelSize.Y; }
+
+private:
+
+	FVector2D PixelSize;
+	double InputMin, InputMax;
+	double OutputMin, OutputMax;
+};
\ No newline at end of file
diff --git a/Engine/Source/Editor/CurveEditor/Public/CurveEditorSelection.h b/Engine/Source/Editor/CurveEditor/Public/CurveEditorSelection.h
new file mode 100644
--- /dev/null
+++ b/Engine/Source/Editor/CurveEditor/Public/CurveEditorSelection.h
@@ -0,0 +1,187 @@
+// Copyright 1998-2018 Epic Games, Inc. All Rights Reserved.
+
+#pragma once
+
+#include "CoreTypes.h"
+#include "Containers/ArrayView.h"
+#include "Containers/Map.h"
+#include "Misc/Optional.h"
+#include "Curves/KeyHandle.h"
+#include "CurveEditorTypes.h"
+
+struct FCurvePointHandle;
+
+
+/**
+ * A set of key handles implemented as a sorted array for transparent passing to TArrayView<> APIs.
+ * Lookup is achieved via binary search: O(log(n)).
+ */
+struct FKeyHandleSet
+{
+	/**
+	 * Add a new key handle to this set
+	 */
+	void Add(FKeyHandle Handle);
+
+	/**
+	 * Remove a handle from this set if it already exists, otherwise add it to the set
+	 */
+	void Toggle(FKeyHandle Handle);
+
+	/**
+	 * Remove a handle from this set
+	 */
+	void Remove(FKeyHandle Handle);
+
+	/**
+	 * Check whether the specified handle exists in this set
+	 */
+	bool Contains(FKeyHandle Handle) const;
+
+	/**
+	 * Retrieve the number of handles in this set
+	 */
+	FORCEINLINE int32 Num() const { return SortedHandles.Num(); }
+
+	/**
+	 * Retrieve a constant view of this set as an array
+	 */
+	FORCEINLINE TArrayView<const FKeyHandle> AsArray() const { return SortedHandles; }
+
+private:
+
+	/** Sorted array of key handles */
+	TArray<FKeyHandle, TInlineAllocator<1>> SortedHandles;
+};
+
+
+/**
+ * Class responsible for tracking selections of keys.
+ * Only one type of point selection is supported at a time (key, arrive tangent, or leave tangent)
+ */
+struct CURVEEDITOR_API FCurveEditorSelection
+{
+	/**
+	 * Default constructor
+	 */
+	FCurveEditorSelection();
+
+	/**
+	 * Retrieve the current type of selection
+	 */
+	FORCEINLINE ECurvePointType GetSelectionType() const { return SelectionType; }
+
+	/**
+	 * Retrieve this selection's serial number. Incremented whenever a change is made to the selection.
+	 */
+	FORCEINLINE uint32 GetSerialNumber() const { return SerialNumber; }
+
+	/**
+	 * Check whether the selection is empty
+	 */
+	FORCEINLINE bool IsEmpty() const { return CurveToSelectedKeys.Num() == 0; }
+
+	/**
+	 * Retrieve all selected key handles, organized by curve ID
+	 */
+	FORCEINLINE const TMap<FCurveModelID, FKeyHandleSet>& GetAll() const { return CurveToSelectedKeys; }
+
+	/**
+	 * Retrieve a set of selected key handles for the specified curve
+	 */
+	const FKeyHandleSet* FindForCurve(FCurveModelID InCurveID) const;
+
+	/**
+	 * Count the total number of selected keys by accumulating the number of selected keys for each curve
+	 */
+	int32 Count() const;
+
+	/**
+	 * Check whether the specified handle is selected
+	 */
+	bool IsSelected(FCurvePointHandle InHandle) const;
+
+	/**
+	 * Check whether the specified handle and curve ID is contained in this selection.
+	 *
+	 * @note: Does not compare the current selection type
+	 */
+	bool Contains(FCurveModelID CurveID, FKeyHandle KeyHandle) const;
+
+public:
+
+	/**
+	 * Add a point handle to this selection, changing the selection type if necessary.
+	 */
+	void Add(FCurvePointHandle InHandle);
+
+	/**
+	 * Add a key handle to this selection, changing the selection type if necessary.
+	 */
+	void Add(FCurveModelID CurveID, ECurvePointType PointType, FKeyHandle KeyHandle);
+
+	/**
+	 * Add key handles to this selection, changing the selection type if necessary.
+	 */
+	void Add(FCurveModelID CurveID, ECurvePointType PointType, TArrayView<const FKeyHandle> Keys);
+
+public:
+
+	/**
+	 * Toggle the selection of the specified point handle, changing the selection type if necessary.
+	 */
+	void Toggle(FCurvePointHandle InHandle);
+
+	/**
+	 * Toggle the selection of the specified key handle, changing the selection type if necessary.
+	 */
+	void Toggle(FCurveModelID CurveID, ECurvePointType PointType, FKeyHandle KeyHandle);
+
+	/**
+	 * Toggle the selection of the specified key handles, changing the selection type if necessary.
+	 */
+	void Toggle(FCurveModelID CurveID, ECurvePointType PointType, TArrayView<const FKeyHandle> Keys);
+
+public:
+
+	/**
+	 * Remove the specified point handle from the selection
+	 */
+	void Remove(FCurvePointHandle InHandle);
+
+	/**
+	 * Remove the specified key handle from the selection
+	 */
+	void Remove(FCurveModelID CurveID, ECurvePointType PointType, FKeyHandle KeyHandle);
+
+	/**
+	 * Remove the specified key handles from the selection
+	 */
+	void Remove(FCurveModelID CurveID, ECurvePointType PointType, TArrayView<const FKeyHandle> Keys);
+
+	/**
+	 * Remove all key handles associated with the specified curve ID from the selection
+	 */
+	void Remove(FCurveModelID InCurveID);
+
+	/**
+	 * Clear the selection entirely
+	 */
+	void Clear();
+
+	/**
+	 * Change the current selection type if it differs from the type specified
+	 */
+	void ChangeSelectionPointType(ECurvePointType InPointType);
+
+private:
+
+	/** A serial number that increments every time a change is made to the selection */
+	uint32 SerialNumber;
+
+	/** The type of point currently selected */
+	ECurvePointType SelectionType;
+
+	/** A map of selected handles stored by curve ID */
+	TMap<FCurveModelID, FKeyHandleSet> CurveToSelectedKeys;
+};
\ No newline at end of file
diff --git a/Engine/Source/Editor/CurveEditor/Public/CurveEditorSnapMetrics.h b/Engine/Source/Editor/CurveEditor/Public/CurveEditorSnapMetrics.h
new file mode 100644
--- /dev/null
+++ b/Engine/Source/Editor/CurveEditor/Public/CurveEditorSnapMetrics.h
@@ -0,0 +1,42 @@
+// Copyright 1998-2018 Epic Games, Inc. All Rights Reserved.
+
+#pragma once
+
+#include "FrameRate.h"
+
+/**
+ * Utility struct that acts as a cache of the current snapping metrics for the curve editor
+ */
+struct FCurveEditorSnapMetrics
+{
+	FCurveEditorSnapMetrics()
+	{
+		bSnapOutputValues = 0;
+		bSnapInputValues = 0;
+		OutputSnapInterval = 1.0;
+	}
+
+	/** Whether we are snapping to the output snap interval */
+	uint8 bSnapOutputValues : 1;
+
+	/** Whether we are snapping to the input snap rate */
+	uint8 bSnapInputValues : 1;
+
+	/** The output snap interval */
+	double OutputSnapInterval;
+
+	/** The input snap rate */
+	FFrameRate InputSnapRate;
+
+	/** Snap the specified input time to the input snap rate if necessary */
+	FORCEINLINE double SnapInputSeconds(double InputTime)
+	{
+		return bSnapInputValues ? (InputTime * InputSnapRate).RoundToFrame() / InputSnapRate : InputTime;
+	}
+
+	/** Snap the specified output value to the output snap interval if necessary */
+	FORCEINLINE double SnapOutput(double OutputValue)
+	{
+		return bSnapOutputValues ? FMath::RoundToDouble(OutputValue / OutputSnapInterval) * OutputSnapInterval : OutputValue;
+	}
+};
diff --git a/Engine/Source/Editor/CurveEditor/Public/ICurveEditorBounds.h b/Engine/Source/Editor/CurveEditor/Public/ICurveEditorBounds.h
new file mode 100644
--- /dev/null
+++ b/Engine/Source/Editor/CurveEditor/Public/ICurveEditorBounds.h
@@ -0,0 +1,56 @@
+// Copyright 1998-2018 Epic Games, Inc. All Rights Reserved.
+
+#pragma once
+
+#include "CoreTypes.h"
+
+struct ICurveEditorBounds
+{
+	virtual ~ICurveEditorBounds(){}
+
+	virtual void GetInputBounds(double& OutMin, double& OutMax) const = 0;
+	virtual void GetOutputBounds(double& OutMin, double& OutMax) const = 0;
+
+	virtual void SetInputBounds(double InMin, double InMax) = 0;
+	virtual void SetOutputBounds(double InMin, double InMax) = 0;
+};
+
+
+struct FStaticCurveEditorBounds : ICurveEditorBounds
+{
+	double InputMin, InputMax;
+	double OutputMin, OutputMax;
+
+	FStaticCurveEditorBounds()
+	{
+		InputMin = 0;
+		InputMax = 1;
+
+		OutputMin = 0;
+		OutputMax = 1;
+	}
+
+	virtual void GetInputBounds(double& OutMin, double& OutMax) const override final
+	{
+		OutMin = InputMin;
+		OutMax = InputMax;
+	}
+
+	virtual void GetOutputBounds(double& OutMin, double& OutMax) const override final
+	{
+		OutMin = OutputMin;
+		OutMax = OutputMax;
+	}
+
+	virtual void SetInputBounds(double InMin, double InMax) override final
+	{
+		InputMin = InMin;
+		InputMax = InMax;
+	}
+
+	virtual void SetOutputBounds(double InMin, double InMax) override final
+	{
+		OutputMin = InMin;
+		OutputMax = InMax;
+	}
+};
\ No newline at end of file
diff --git a/Engine/Source/Editor/CurveEditor/Public/ICurveEditorDragOperation.h b/Engine/Source/Editor/CurveEditor/Public/ICurveEditorDragOperation.h
new file mode 100644
--- /dev/null
+++ b/Engine/Source/Editor/CurveEditor/Public/ICurveEditorDragOperation.h
@@ -0,0 +1,138 @@
+// Copyright 1998-2018 Epic Games, Inc. All Rights Reserved.
+
+#pragma once
+
+#include "Framework/DelayedDrag.h"
+#include "Misc/Optional.h"
+#include "ScopedTransaction.h"
+#include "CurveEditorTypes.h"
+#include "CurveEditorSnapMetrics.h"
+
+class FCurveEditor;
+
+/**
+ * Interface for all drag operations in the curve editor
+ */
+class ICurveEditorDragOperation
+{
+public:
+
+	ICurveEditorDragOperation()
+		: MouseLockVector(FVector2D::UnitVector)
+	{}
+
+	virtual ~ICurveEditorDragOperation() {}
+
+	/**
+	 * Begin this drag operation with the specified initial and current positions
+	 */
+	void BeginDrag(FVector2D InitialPosition, FVector2D CurrentPosition, const FPointerEvent& MouseEvent);
+
+	/**
+	 * Continue this drag operation with the specified initial and current positions
+	 */
+	void Drag(FVector2D InitialPosition, FVector2D CurrentPosition, const FPointerEvent& MouseEvent);
+
+	/**
+	 * Finish this drag operation with the specified initial and current positions
+	 */
+	void EndDrag(FVector2D InitialPosition, FVector2D CurrentPosition, const FPointerEvent& MouseEvent);
+
+	/**
+	 * Paint this drag operation
+	 */
+	int32 Paint(const FGeometry& AllottedGeometry, FSlateWindowElementList& OutDrawElements, int32 LayerId);
+
+	/**
+	 * Cancel this drag operation
+	 */
+	void CancelDrag();
+
+protected:
+
+	/** Implementation method for derived types to begin a drag */
+	virtual void OnBeginDrag(FVector2D InitialPosition, FVector2D CurrentPosition, const FPointerEvent& MouseEvent)
+	{}
+
+	/** Implementation method for derived types to continue a drag */
+	virtual void OnDrag(FVector2D InitialPosition, FVector2D CurrentPosition, const FPointerEvent& MouseEvent)
+	{}
+
+	/** Implementation method for derived types to finish a drag */
+	virtual void OnEndDrag(FVector2D InitialPosition, FVector2D CurrentPosition, const FPointerEvent& MouseEvent)
+	{}
+
+	/** Implementation method for derived types to paint this drag */
+	virtual int32 OnPaint(const FGeometry& AllottedGeometry, FSlateWindowElementList& OutDrawElements, int32 LayerId)
+	{
+		return LayerId;
+	}
+
+	/** Implementation method for derived types to cancel a drag */
+	virtual void OnCancelDrag()
+	{}
+
+	/**
+	 * Determine the effective mouse position for a drag vector, potentially locked to an axis based on the current pointer event
+	 */
+	FVector2D GetLockedMousePosition(FVector2D InitialPosition, FVector2D CurrentPosition, const FPointerEvent& MouseEvent);
+
+protected:
+
+	/** The vector we're currently locked to, or (1.f, 1.f) if we're not locked */
+	FVector2D MouseLockVector;
+};
+
+/**
+ * Interface for all key drag operations in the curve editor
+ */
+class ICurveEditorKeyDragOperation : public ICurveEditorDragOperation
+{
+public:
+
+	/** Cached (and potentially manipulated) snap metrics to be used for this drag */
+	FCurveEditorSnapMetrics SnapMetrics;
+
+	/**
+	 * Initialize this drag operation using the specified curve editor pointer and an optional cardinal point
+
+	 * @param InCurveEditor       Curve editor pointer. Guaranteed to persist for the lifetime of this drag.
+	 * @param CardinalPoint       The point that should be considered the origin of this drag.
+	 */
+	void Initialize(FCurveEditor* InCurveEditor, const TOptional<FCurvePointHandle>& CardinalPoint);
+
+protected:
+
+	/** Implementation method for derived types to initialize a drag */
+	virtual void OnInitialize(FCurveEditor* InCurveEditor, const TOptional<FCurvePointHandle>& CardinalPoint)
+	{}
+
+	virtual void OnCancelDrag() override
+	{
+		if (Transaction.IsValid())
+		{
+			Transaction->Cancel();
+		}
+	}
+
+	/** Scoped transaction pointer */
+	TUniquePtr<FScopedTransaction> Transaction;
+};
+
+/**
+ * Utility struct used to facilitate a delayed drag operation with an implementation interface
+ */
+struct FCurveEditorDelayedDrag : FDelayedDrag
+{
+	/**
+	 * The drag implementation to use once the drag has started
+	 */
+	TUniquePtr<ICurveEditorDragOperation> DragImpl;
+
+	/**
+	 * Start a delayed drag operation at the specified position and effective key
+	 */
+	FCurveEditorDelayedDrag(FVector2D InInitialPosition, FKey InEffectiveKey)
+		: FDelayedDrag(InInitialPosition, InEffectiveKey)
+	{}
+};
diff --git a/Engine/Source/Editor/CurveEditor/Public/SCurveEditorPanel.h b/Engine/Source/Editor/CurveEditor/Public/SCurveEditorPanel.h
new file mode 100644
--- /dev/null
+++ b/Engine/Source/Editor/CurveEditor/Public/SCurveEditorPanel.h
@@ -0,0 +1,234 @@
+// Copyright 1998-2018 Epic Games, Inc. All Rights Reserved.
+
+#pragma once
+
+#include "Widgets/DeclarativeSyntaxSupport.h"
+#include "Widgets/SCompoundWidget.h"
+
+#include "CurveEditorScreenSpace.h"
+#include "ICurveEditorDragOperation.h"
+#include "CurveDrawInfo.h"
+#include "CurveEditor.h"
+#include "CurveEditorTypes.h"
+
+class SScrollBox;
+
+/**
+ * Curve editor widget that reflects the state of an FCurveEditor
+ */
+class CURVEEDITOR_API SCurveEditorPanel : public SCompoundWidget
+{
+	SLATE_BEGIN_ARGS(SCurveEditorPanel)
+		: _GridLineTint(FLinearColor(0.1f, 0.1f, 0.1f, 1.f))
+	{}
+
+		/** Color to draw grid lines */
+		SLATE_ATTRIBUTE(FLinearColor, GridLineTint)
+
+	SLATE_END_ARGS()
+
+	/**
+	 * Construct a new curve editor panel widget
+	 */
+	void Construct(const FArguments& InArgs, TSharedRef<FCurveEditor> InCurveEditor);
+
+	/**
+	 * Access the draw parameters that this curve editor has cached for this frame
+	 */
+	const TArray<FCurveDrawParams>& GetCachedDrawParams() const
+	{
+		return CachedDrawParams;
+	}
+
+	/**
+	 * Access the combined command list for this curve editor and panel widget
+	 */
+	TSharedPtr<FUICommandList> GetCommands() const
+	{
+		return CommandList;
+	}
+
+private:
+
+	virtual int32 OnPaint(const FPaintArgs& Args, const FGeometry& AllottedGeometry, const FSlateRect& MyCullingRect, FSlateWindowElementList& OutDrawElements, int32 LayerId, const FWidgetStyle& InWidgetStyle, bool bParentEnabled) const override;
+	virtual void Tick(const FGeometry& AllottedGeometry, const double InCurrentTime, const float InDeltaTime) override;
+
+	/*~ Mouse interaction */
+	virtual FReply OnMouseButtonDown(const FGeometry& MyGeometry, const FPointerEvent& MouseEvent) override;
+	virtual FReply OnMouseButtonUp(const FGeometry& MyGeometry, const FPointerEvent& MouseEvent) override;
+	virtual FReply OnMouseMove(const FGeometry& MyGeometry, const FPointerEvent& MouseEvent) override;
+	virtual FReply OnMouseWheel(const FGeometry& MyGeometry, const FPointerEvent& MouseEvent) override;
+	virtual void OnFocusLost(const FFocusEvent& InFocusEvent) override;
+	virtual void OnMouseEnter(const FGeometry& MyGeometry, const FPointerEvent& MouseEvent) override;
+	virtual void OnMouseLeave(const FPointerEvent& MouseEvent) override;
+
+	/*~ Keyboard interaction */
+	virtual bool SupportsKeyboardFocus() const override { return true; }
+	virtual FReply OnKeyDown(const FGeometry& MyGeometry, const FKeyEvent& InKeyEvent) override;
+
+private:
+
+	/**
+	 * Get the visibility for the value splitter control
+	 */
+	EVisibility GetSplitterVisibility() const;
+
+	/**
+	 * Get the currently hovered curve ID
+	 */
+	TOptional<FCurveModelID> GetHoveredCurve() const;
+
+	/**
+	 * Hit test for a point on the curve editor using a mouse position in slate units
+	 */
+	TOptional<FCurvePointHandle> HitPoint(FVector2D MousePixel) const;
+
+	/*~ Event bindings */
+	void UpdateEditBox();
+	void UpdateCommonCurveInfo();
+	void UpdateToolTip(const FGeometry& InMyGeometry, const FPointerEvent& InMouseEvent);
+	void UpdateCurveProximities(FVector2D MousePixel);
+
+private:
+
+	/**
+	 * Bind command mappings for this widget
+	 */
+	void BindCommands();
+
+	/**
+	 * Rebind contextual command mappings that rely on the mouse position
+	 */
+	void RebindContextualActions(FVector2D MousePosition);
+
+	/*~ Command binding callbacks */
+	void OnAddKey(FVector2D MousePixel);
+	void OnAddKeyHere(FVector2D MousePixel);
+	void OnAddKey(FCurveModelID CurveToAdd, FVector2D MousePixel);
+
+	/**
+	 * Assign new attributes to the currently selected keys
+	 */
+	void SetKeyAttributes(FKeyAttributes KeyAttributes, FText Description);
+
+	/**
+	 * Assign new curve attributes to all visible curves
+	 */
+	void SetCurveAttributes(FCurveAttributes CurveAttributes, FText Description);
+
+	/** Compare all the currently selected keys' interp modes against the specified interp mode */
+	bool CompareCommonInterpolationMode(ERichCurveInterpMode InterpMode) const
+	{
+		return CachedCommonKeyAttributes.HasInterpMode() && CachedCommonKeyAttributes.GetInterpMode() == InterpMode;
+	}
+
+	/** Compare all the currently selected keys' tangent modes against the specified tangent mode */
+	bool CompareCommonTangentMode(ERichCurveInterpMode InterpMode, ERichCurveTangentMode TangentMode) const
+	{
+		return CompareCommonInterpolationMode(InterpMode) && CachedCommonKeyAttributes.HasTangentMode() && CachedCommonKeyAttributes.GetTangentMode() == TangentMode;
+	}
+
+	/** Compare all the visible curves' pre-extrapolation modes against the specified extrapolation mode */
+	bool CompareCommonPreExtrapolationMode(ERichCurveExtrapolation PreExtrapolationMode) const
+	{
+		return CachedCommonCurveAttributes.HasPreExtrapolation()  && CachedCommonCurveAttributes.GetPreExtrapolation() == PreExtrapolationMode;
+	}
+
+	/** Compare all the visible curves' post-extrapolation modes against the specified extrapolation mode */
+	bool CompareCommonPostExtrapolationMode(ERichCurveExtrapolation PostExtrapolationMode) const
+	{
+		return CachedCommonCurveAttributes.HasPostExtrapolation() && CachedCommonCurveAttributes.GetPostExtrapolation() == PostExtrapolationMode;
+	}
+
+
+	/**
+	*  Brings up Dialog to specify Tolerance before doing the key reduction.
+	*/
+	void OnSimplifySelection();
+
+	/**
+	*  Callback for when the curve is simplified.
+	*/
+	void OnSimplifySelectionCommited(const FText& InText, ETextCommit::Type CommitInfo);
+
+	/**
+	*  Create popup menu
+	*/
+	void GenericTextEntryModeless(const FText& DialogText, const FText& DefaultText, FOnTextCommitted OnTextComitted);
+
+
+	/**
+	*  Close popup menu
+	*/
+	void CloseEntryPopupMenu();
+public:
+
+	/**
+	 * Draw grid lines
+	 */
+	int32 DrawGridLines(const FGeometry& AllottedGeometry, FSlateWindowElementList& OutDrawElements, int32 LayerId, ESlateDrawEffect DrawEffects) const;
+
+	/**
+	 * Draw curve data
+	 */
+	int32 DrawCurves(const FGeometry& AllottedGeometry, FSlateWindowElementList& OutDrawElements, int32 LayerId, const FWidgetStyle& InWidgetStyle, ESlateDrawEffect DrawEffects) const;
+
+private:
+
+	/** The curve editor pointer */
+	TSharedPtr<FCurveEditor> CurveEditor;
+
+	/** Array of curve proximities in slate units that's updated on mouse move */
+	TArray<TTuple<FCurveModelID, float>> CurveProximities;
+
+	/** (Optional) the current drag operation */
+	TOptional<FCurveEditorDelayedDrag> DragOperation;
+
+	/** Curve draw parameters that are re-generated on tick */
+	TArray<FCurveDrawParams> CachedDrawParams;
+
+	/** Cached curve attributes that are common to all visible curves */
+	FCurveAttributes CachedCommonCurveAttributes;
+
+	/** Cached key attributes that are common to all selected keys */
+	FKeyAttributes CachedCommonKeyAttributes;
+
+	/** Attribute used for retrieving the desired grid line color */
+	TAttribute<FLinearColor> GridLineTintAttribute;
+
+	/** Edit panel scroll box */
+	TSharedPtr<SScrollBox> EditPanel;
+
+	/** Map of edit UI widgets for each curve in the current selection set */
+	TMap<FCurveModelID, TSharedPtr<SWidget>> CurveToEditUI;
+
+	/** Command list for widget specific command bindings */
+	TSharedPtr<FUICommandList> CommandList;
+
+	/** Cached serial number from the curve editor selection. Used to update edit UIs when the selection changes. */
+	uint32 CachedSelectionSerialNumber;
+
+private:
+
+	/*~ Cached tooltip data */
+	bool  IsToolTipEnabled() const;
+	FText GetToolTipCurveName() const;
+	FText GetToolTipTimeText() const;
+	FText GetToolTipValueText() const;
+
+	struct FCachedToolTipData
+	{
+		FCachedToolTipData() {}
+
+		FText Text;
+		FText EvaluatedValue;
+		FText EvaluatedTime;
+	};
+	TOptional<FCachedToolTipData> CachedToolTipData;
+
+	/** The tolerance to use when reducing curves */
+	float ReduceTolerance;
+
+	/** Generic Popup Entry */
+	TWeakPtr<IMenu> EntryPopupMenu;
+};
\ No newline at end of file
diff --git a/Engine/Source/Editor/DetailCustomizations/Private/CameraDetails.cpp b/Engine/Source/Editor/DetailCustomizations/Private/CameraDetails.cpp
--- a/Engine/Source/Editor/DetailCustomizations/Private/CameraDetails.cpp
+++ b/Engine/Source/Editor/DetailCustomizations/Private/CameraDetails.cpp
@@ -98,7 +98,6 @@ void FCameraDetails::CustomizeDetails( IDetailLayoutBuilder& DetailLayout )
 		OrthoFarClipPlaneRow.Visibility(OrthographicVisibility);
 
 		// Aspect ratio
-		CameraCategory.AddProperty( bConstrainAspectRatioProperty );
 		IDetailPropertyRow& AspectRatioRow = CameraCategory.AddProperty(AspectRatioProperty);
 
 		// Provide the special aspect ratio row
@@ -142,11 +141,13 @@ void FCameraDetails::CustomizeDetails( IDetailLayoutBuilder& DetailLayout )
 					]
 				]
 			];
-		
-		CameraCategory.AddProperty(DetailLayout.GetProperty(GET_MEMBER_NAME_CHECKED(UCameraComponent, bUsePawnControlRotation)));
-		CameraCategory.AddProperty(DetailLayout.GetProperty(GET_MEMBER_NAME_CHECKED(UCameraComponent, PostProcessBlendWeight)));
 	}
 
+	IDetailCategoryBuilder& CameraSettingsCategory = DetailLayout.EditCategory( "CameraOptions", FText::GetEmpty(), ECategoryPriority::Important );
+	CameraSettingsCategory.AddProperty( bConstrainAspectRatioProperty );
+	CameraSettingsCategory.AddProperty(DetailLayout.GetProperty(GET_MEMBER_NAME_CHECKED(UCameraComponent, bUsePawnControlRotation)));
+	CameraSettingsCategory.AddProperty(DetailLayout.GetProperty(GET_MEMBER_NAME_CHECKED(UCameraComponent, PostProcessBlendWeight)));
+
 	UpdateAspectTextFromProperty();
 }
 END_SLATE_FUNCTION_BUILD_OPTIMIZATION
diff --git a/Engine/Source/Editor/DetailCustomizations/Private/ComponentTransformDetails.cpp b/Engine/Source/Editor/DetailCustomizations/Private/ComponentTransformDetails.cpp
--- a/Engine/Source/Editor/DetailCustomizations/Private/ComponentTransformDetails.cpp
+++ b/Engine/Source/Editor/DetailCustomizations/Private/ComponentTransformDetails.cpp
@@ -1146,10 +1146,7 @@ void FComponentTransformDetails::OnSetTransform(ETransformField::Type TransformF
 								ObjectToRelativeRotationMap.FindOrAdd(SceneComponent) = NewRotation;
 							}
 
-							SceneComponent->SetRelativeRotation(NewRotation);
-
-							// Also forcibly set it as the cache may have changed it slightly
-							SceneComponent->RelativeRotation = NewRotation;
+							SceneComponent->SetRelativeRotationExact(NewRotation);
 
 							// If it's a template, propagate the change out to any current instances of the object
 							if (bIsEditingTemplateObject)
diff --git a/Engine/Source/Editor/DetailCustomizations/Private/DetailCustomizations.cpp b/Engine/Source/Editor/DetailCustomizations/Private/DetailCustomizations.cpp
--- a/Engine/Source/Editor/DetailCustomizations/Private/DetailCustomizations.cpp
+++ b/Engine/Source/Editor/DetailCustomizations/Private/DetailCustomizations.cpp
@@ -108,8 +108,9 @@
 #include "MovieSceneCaptureCustomization.h"
 #include "MovieSceneEvalOptionsCustomization.h"
 #include "MovieSceneEventParametersCustomization.h"
+#include "FrameRateCustomization.h"
 #include "MovieSceneSequencePlaybackSettingsCustomization.h"
-#include "MovieSceneCurveInterfaceKeyEditStructCustomization.h"
+#include "SequencerGenericKeyStructCustomization.h"
 #include "LevelSequenceBurnInOptionsCustomization.h"
 #include "MovieSceneBindingOverrideDataCustomization.h"
 #include "TextCustomization.h"
@@ -255,6 +256,7 @@ void FDetailCustomizationsModule::RegisterPropertyTypeCustomizations()
 	RegisterCustomPropertyTypeLayout("MovieSceneTrackEvalOptions", FOnGetPropertyTypeCustomizationInstance::CreateStatic(&FMovieSceneTrackEvalOptionsCustomization::MakeInstance));
 	RegisterCustomPropertyTypeLayout("MovieSceneSectionEvalOptions", FOnGetPropertyTypeCustomizationInstance::CreateStatic(&FMovieSceneSectionEvalOptionsCustomization::MakeInstance));
 	RegisterCustomPropertyTypeLayout("MovieSceneEventParameters", FOnGetPropertyTypeCustomizationInstance::CreateStatic(&FMovieSceneEventParametersCustomization::MakeInstance));
+	RegisterCustomPropertyTypeLayout("FrameRate", FOnGetPropertyTypeCustomizationInstance::CreateStatic(&FFrameRateCustomization::MakeInstance));
 	RegisterCustomPropertyTypeLayout("LevelSequenceBurnInOptions", FOnGetPropertyTypeCustomizationInstance::CreateStatic(&FLevelSequenceBurnInOptionsCustomization::MakeInstance));
 	RegisterCustomPropertyTypeLayout("LevelSequenceBurnInInitSettings", FOnGetPropertyTypeCustomizationInstance::CreateStatic(&FLevelSequenceBurnInInitSettingsCustomization::MakeInstance));
 	RegisterCustomPropertyTypeLayout("CollectionReference", FOnGetPropertyTypeCustomizationInstance::CreateStatic(&FCollectionReferenceStructCustomization::MakeInstance));
@@ -354,7 +356,7 @@ void FDetailCustomizationsModule::RegisterObjectCustomizations()
 	RegisterCustomClassLayout("ConfigHierarchyPropertyView", FOnGetDetailCustomizationInstance::CreateStatic(&FConfigPropertyHelperDetails::MakeInstance));
 
 	RegisterCustomClassLayout("MovieSceneCapture", FOnGetDetailCustomizationInstance::CreateStatic(&FMovieSceneCaptureCustomization::MakeInstance));
-	RegisterCustomClassLayout("MovieSceneCurveInterfaceKeyEditStruct", FOnGetDetailCustomizationInstance::CreateStatic(&FMovieSceneCurveInterfaceKeyEditStructCustomization::MakeInstance));
+	RegisterCustomClassLayout("SequencerGenericKeyStruct", FOnGetDetailCustomizationInstance::CreateStatic(&FSequencerGenericKeyStructCustomization::MakeInstance));
 
 	RegisterCustomClassLayout("AnalyticsPrivacySettings", FOnGetDetailCustomizationInstance::CreateStatic(&FImportantToggleSettingCustomization::MakeInstance));
 	RegisterCustomClassLayout("EndUserSettings", FOnGetDetailCustomizationInstance::CreateStatic(&FImportantToggleSettingCustomization::MakeInstance));
diff --git a/Engine/Source/Editor/DetailCustomizations/Private/SequencerGenericKeyStructCustomization.cpp b/Engine/Source/Editor/DetailCustomizations/Private/SequencerGenericKeyStructCustomization.cpp
new file mode 100644
--- /dev/null
+++ b/Engine/Source/Editor/DetailCustomizations/Private/SequencerGenericKeyStructCustomization.cpp
@@ -0,0 +1,31 @@
+// Copyright 1998-2018 Epic Games, Inc. All Rights Reserved.
+
+#include "SequencerGenericKeyStructCustomization.h"
+
+#include "DetailLayoutBuilder.h"
+#include "SequencerGenericKeyStruct.h"
+
+#define LOCTEXT_NAMESPACE "SequencerGenericKeyStructCustomization"
+
+TSharedRef<IDetailCustomization> FSequencerGenericKeyStructCustomization::MakeInstance()
+{
+	return MakeShared<FSequencerGenericKeyStructCustomization>();
+}
+
+void FSequencerGenericKeyStructCustomization::CustomizeDetails(IDetailLayoutBuilder& DetailBuilder)
+{
+	TArray<TSharedPtr<FStructOnScope>> Structs;
+	DetailBuilder.GetStructsBeingCustomized(Structs);
+
+	const UStruct* StructType = Structs.Num() == 1 ? Structs[0]->GetStruct() : nullptr;
+	if (StructType && StructType == FSequencerGenericKeyStruct::StaticStruct())
+	{
+		auto* StructValue = (FSequencerGenericKeyStruct*)Structs[0]->GetStructMemory();
+		if (StructValue && StructValue->CustomizationImpl.IsValid())
+		{
+			StructValue->CustomizationImpl->Extend(DetailBuilder);
+		}
+	}
+}
+
+#undef LOCTEXT_NAMESPACE
diff --git a/Engine/Source/Editor/DetailCustomizations/Private/SequencerGenericKeyStructCustomization.h b/Engine/Source/Editor/DetailCustomizations/Private/SequencerGenericKeyStructCustomization.h
new file mode 100644
--- /dev/null
+++ b/Engine/Source/Editor/DetailCustomizations/Private/SequencerGenericKeyStructCustomization.h
@@ -0,0 +1,16 @@
+// Copyright 1998-2018 Epic Games, Inc. All Rights Reserved.
+
+#pragma once
+
+#include "CoreMinimal.h"
+#include "IDetailCustomization.h"
+
+class IDetailLayoutBuilder;
+
+class FSequencerGenericKeyStructCustomization : public IDetailCustomization
+{
+public:
+	static TSharedRef<IDetailCustomization> MakeInstance();
+
+	virtual void CustomizeDetails(IDetailLayoutBuilder& DetailBuilder) override;
+};
diff --git a/Engine/Source/Editor/EditorSettingsViewer/EditorSettingsViewer.Build.cs b/Engine/Source/Editor/EditorSettingsViewer/EditorSettingsViewer.Build.cs
--- a/Engine/Source/Editor/EditorSettingsViewer/EditorSettingsViewer.Build.cs
+++ b/Engine/Source/Editor/EditorSettingsViewer/EditorSettingsViewer.Build.cs
@@ -15,6 +15,7 @@ namespace UnrealBuildTool.Rules
 			PrivateDependencyModuleNames.AddRange(
 				new string[] {
 					"CoreUObject",
+					"CurveEditor",
 					"Engine",
 					"GraphEditor",
 					"InputBindingEditor",
diff --git a/Engine/Source/Editor/EditorStyle/Private/SlateEditorStyle.cpp b/Engine/Source/Editor/EditorStyle/Private/SlateEditorStyle.cpp
--- a/Engine/Source/Editor/EditorStyle/Private/SlateEditorStyle.cpp
+++ b/Engine/Source/Editor/EditorStyle/Private/SlateEditorStyle.cpp
@@ -1748,10 +1748,14 @@ void FSlateEditorStyle::FStyle::SetupGeneralStyles()
 		Set( "Sequencer.PartialKey", new IMAGE_BRUSH( "Sequencer/PartialKey", FVector2D(11.f, 11.f) ) );
 		Set( "Sequencer.Star", new IMAGE_BRUSH( "Sequencer/Star", Icon12x12 ) );
 		Set( "Sequencer.Empty", new IMAGE_BRUSH( "Sequencer/Empty", Icon12x12 ) );
+		Set( "Sequencer.TangentHandle", new IMAGE_BRUSH( "Sequencer/TangentHandle", FVector2D(7, 7) ) );
 		Set( "Sequencer.GenericDivider", new IMAGE_BRUSH( "Sequencer/GenericDivider", FVector2D(2.f, 2.f), FLinearColor::White, ESlateBrushTileType::Vertical ) );
 
+		Set( "Sequencer.Timeline.ScrubFill", new BOX_BRUSH( "Sequencer/ScrubFill", FMargin( 2.f/4.f, 0.f ) ) );
 		Set( "Sequencer.Timeline.ScrubHandleDown", new BOX_BRUSH( "Sequencer/ScrubHandleDown", FMargin( 6.f/13.f, 5/12.f, 6/13.f, 8/12.f ) ) );
 		Set( "Sequencer.Timeline.ScrubHandleUp", new BOX_BRUSH( "Sequencer/ScrubHandleUp", FMargin( 6.f/13.f, 8/12.f, 6/13.f, 5/12.f ) ) );
+		Set( "Sequencer.Timeline.ClampedScrubHandleDown", new BOX_BRUSH( "Sequencer/ScrubHandleDown_Clamped", FMargin( 6.f/13.f, 3.f/12.f, 6.f/13.f, 7.f/12.f ) ) );
+		Set( "Sequencer.Timeline.ClampedScrubHandleUp", new BOX_BRUSH( "Sequencer/ScrubHandleUp_Clamped", FMargin( 6.f/13.f, 8/12.f, 6/13.f, 5/12.f ) ) );
 		Set( "Sequencer.Timeline.ScrubHandleWhole", new BOX_BRUSH( "Sequencer/ScrubHandleWhole", FMargin( 6.f/13.f, 10/24.f, 6/13.f, 10/24.f  ) ) );
 		Set( "Sequencer.Timeline.RangeHandleLeft", new BOX_BRUSH( "Sequencer/GenericGripLeft", FMargin(5.f/16.f) ) );
 		Set( "Sequencer.Timeline.RangeHandleRight", new BOX_BRUSH( "Sequencer/GenericGripRight", FMargin(5.f/16.f) ) );
@@ -1888,8 +1892,11 @@ void FSlateEditorStyle::FStyle::SetupGeneralStyles()
 		Set( "Sequencer.Value", new IMAGE_BRUSH( "Sequencer/Main_Icons/Icon_Sequencer_Value_24x", Icon48x48 ) );
 		Set( "Sequencer.Value.Small", new IMAGE_BRUSH( "Sequencer/Main_Icons/Icon_Sequencer_Value_24x", Icon24x24 ) );
 
+		Set( "Sequencer.OverlayPanel.Background", new BOX_BRUSH( "Sequencer/OverlayPanelBackground", FMargin(26.f/54.f) ) );
+
 		Set( "Sequencer.TrackArea.LaneColor", FLinearColor(0.3f, 0.3f, 0.3f, 0.3f) );
 
+		Set( "Sequencer.Tracks.Media", new IMAGE_BRUSH("Sequencer/Dropdown_Icons/Icon_Media_Track_16x", Icon16x16));
 		Set( "Sequencer.Tracks.Audio", new IMAGE_BRUSH("Sequencer/Dropdown_Icons/Icon_Audio_Track_16x", Icon16x16));
 		Set( "Sequencer.Tracks.Event", new IMAGE_BRUSH("Sequencer/Dropdown_Icons/Icon_Event_Track_16x", Icon16x16));
 		Set( "Sequencer.Tracks.Fade", new IMAGE_BRUSH("Sequencer/Dropdown_Icons/Icon_Fade_Track_16x", Icon16x16));
@@ -6884,35 +6891,35 @@ void FSlateEditorStyle::FStyle::SetupToolkitStyles()
 		Set( "CurveEditor.DeleteTab.Small", new IMAGE_BRUSH( "Icons/icon_CurveEditor_DeleteTab_40x", Icon20x20 ) );
 	}
 
-	// Rich Curve Editor
-	{
-		Set("RichCurveEditor.ZoomToFitHorizontal", new IMAGE_BRUSH("Icons/icon_CurveEditor_Horizontal_40x", Icon40x40));
-		Set("RichCurveEditor.ZoomToFitHorizontal.Small", new IMAGE_BRUSH("Icons/icon_CurveEditor_Horizontal_40x", Icon20x20));
-		Set("RichCurveEditor.ZoomToFitVertical", new IMAGE_BRUSH("Icons/icon_CurveEditor_Vertical_40x", Icon40x40));
-		Set("RichCurveEditor.ZoomToFitVertical.Small", new IMAGE_BRUSH("Icons/icon_CurveEditor_Vertical_40x", Icon20x20));
-		Set("RichCurveEditor.ZoomToFit", new IMAGE_BRUSH("Icons/icon_CurveEditor_ZoomToFit_40x", Icon40x40));
-		Set("RichCurveEditor.ZoomToFit.Small", new IMAGE_BRUSH("Icons/icon_CurveEditor_ZoomToFit_40x", Icon20x20));
-
-		Set("RichCurveEditor.ToggleInputSnapping", new IMAGE_BRUSH("Icons/icon_CurveEditor_ToggleInputSnap_40x", Icon40x40));
-		Set("RichCurveEditor.ToggleInputSnapping.Small", new IMAGE_BRUSH("Icons/icon_CurveEditor_ToggleInputSnap_40x", Icon20x20));
-		Set("RichCurveEditor.ToggleOutputSnapping", new IMAGE_BRUSH("Icons/icon_CurveEditor_ToggleOutputSnap_40x", Icon40x40));
-		Set("RichCurveEditor.ToggleOutputSnapping.Small", new IMAGE_BRUSH("Icons/icon_CurveEditor_ToggleOutputSnap_40x", Icon20x20));
-
-		Set("RichCurveEditor.InterpolationCubicAuto", new IMAGE_BRUSH("Icons/icon_CurveEditor_Auto_40x", Icon40x40));
-		Set("RichCurveEditor.InterpolationCubicAuto.Small", new IMAGE_BRUSH("Icons/icon_CurveEditor_Auto_40x", Icon20x20));
-		Set("RichCurveEditor.InterpolationCubicUser", new IMAGE_BRUSH("Icons/icon_CurveEditor_User_40x", Icon40x40));
-		Set("RichCurveEditor.InterpolationCubicUser.Small", new IMAGE_BRUSH("Icons/icon_CurveEditor_User_40x", Icon20x20));
-		Set("RichCurveEditor.InterpolationCubicBreak", new IMAGE_BRUSH("Icons/icon_CurveEditor_Break_40x", Icon40x40));
-		Set("RichCurveEditor.InterpolationCubicBreak.Small", new IMAGE_BRUSH("Icons/icon_CurveEditor_Break_40x", Icon20x20));
-		Set("RichCurveEditor.InterpolationLinear", new IMAGE_BRUSH("Icons/icon_CurveEditor_Linear_40x", Icon40x40));
-		Set("RichCurveEditor.InterpolationLinear.Small", new IMAGE_BRUSH("Icons/icon_CurveEditor_Linear_40x", Icon20x20));
-		Set("RichCurveEditor.InterpolationConstant", new IMAGE_BRUSH("Icons/icon_CurveEditor_Constant_40x", Icon40x40));
-		Set("RichCurveEditor.InterpolationConstant.Small", new IMAGE_BRUSH("Icons/icon_CurveEditor_Constant_40x", Icon20x20));
-
-		Set("RichCurveEditor.FlattenTangents", new IMAGE_BRUSH("Icons/icon_CurveEditor_Flatten_40x", Icon40x40));
-		Set("RichCurveEditor.FlattenTangents.Small", new IMAGE_BRUSH("Icons/icon_CurveEditor_Flatten_40x", Icon20x20));
-		Set("RichCurveEditor.StraightenTangents", new IMAGE_BRUSH("Icons/icon_CurveEditor_Straighten_40x", Icon40x40));
-		Set("RichCurveEditor.StraightenTangents.Small", new IMAGE_BRUSH("Icons/icon_CurveEditor_Straighten_40x", Icon20x20));
+	// Curve Editor
+	{
+		Set("GenericCurveEditor.ZoomToFitHorizontal", new IMAGE_BRUSH("Icons/icon_CurveEditor_Horizontal_40x", Icon40x40));
+		Set("GenericCurveEditor.ZoomToFitHorizontal.Small", new IMAGE_BRUSH("Icons/icon_CurveEditor_Horizontal_40x", Icon20x20));
+		Set("GenericCurveEditor.ZoomToFitVertical", new IMAGE_BRUSH("Icons/icon_CurveEditor_Vertical_40x", Icon40x40));
+		Set("GenericCurveEditor.ZoomToFitVertical.Small", new IMAGE_BRUSH("Icons/icon_CurveEditor_Vertical_40x", Icon20x20));
+		Set("GenericCurveEditor.ZoomToFit", new IMAGE_BRUSH("Icons/icon_CurveEditor_ZoomToFit_40x", Icon40x40));
+		Set("GenericCurveEditor.ZoomToFit.Small", new IMAGE_BRUSH("Icons/icon_CurveEditor_ZoomToFit_40x", Icon20x20));
+
+		Set("GenericCurveEditor.ToggleInputSnapping", new IMAGE_BRUSH("Icons/icon_CurveEditor_ToggleInputSnap_40x", Icon40x40));
+		Set("GenericCurveEditor.ToggleInputSnapping.Small", new IMAGE_BRUSH("Icons/icon_CurveEditor_ToggleInputSnap_40x", Icon20x20));
+		Set("GenericCurveEditor.ToggleOutputSnapping", new IMAGE_BRUSH("Icons/icon_CurveEditor_ToggleOutputSnap_40x", Icon40x40));
+		Set("GenericCurveEditor.ToggleOutputSnapping.Small", new IMAGE_BRUSH("Icons/icon_CurveEditor_ToggleOutputSnap_40x", Icon20x20));
+
+		Set("GenericCurveEditor.InterpolationCubicAuto", new IMAGE_BRUSH("Icons/icon_CurveEditor_Auto_40x", Icon40x40));
+		Set("GenericCurveEditor.InterpolationCubicAuto.Small", new IMAGE_BRUSH("Icons/icon_CurveEditor_Auto_40x", Icon20x20));
+		Set("GenericCurveEditor.InterpolationCubicUser", new IMAGE_BRUSH("Icons/icon_CurveEditor_User_40x", Icon40x40));
+		Set("GenericCurveEditor.InterpolationCubicUser.Small", new IMAGE_BRUSH("Icons/icon_CurveEditor_User_40x", Icon20x20));
+		Set("GenericCurveEditor.InterpolationCubicBreak", new IMAGE_BRUSH("Icons/icon_CurveEditor_Break_40x", Icon40x40));
+		Set("GenericCurveEditor.InterpolationCubicBreak.Small", new IMAGE_BRUSH("Icons/icon_CurveEditor_Break_40x", Icon20x20));
+		Set("GenericCurveEditor.InterpolationLinear", new IMAGE_BRUSH("Icons/icon_CurveEditor_Linear_40x", Icon40x40));
+		Set("GenericCurveEditor.InterpolationLinear.Small", new IMAGE_BRUSH("Icons/icon_CurveEditor_Linear_40x", Icon20x20));
+		Set("GenericCurveEditor.InterpolationConstant", new IMAGE_BRUSH("Icons/icon_CurveEditor_Constant_40x", Icon40x40));
+		Set("GenericCurveEditor.InterpolationConstant.Small", new IMAGE_BRUSH("Icons/icon_CurveEditor_Constant_40x", Icon20x20));
+
+		Set("GenericCurveEditor.FlattenTangents", new IMAGE_BRUSH("Icons/icon_CurveEditor_Flatten_40x", Icon40x40));
+		Set("GenericCurveEditor.FlattenTangents.Small", new IMAGE_BRUSH("Icons/icon_CurveEditor_Flatten_40x", Icon20x20));
+		Set("GenericCurveEditor.StraightenTangents", new IMAGE_BRUSH("Icons/icon_CurveEditor_Straighten_40x", Icon40x40));
+		Set("GenericCurveEditor.StraightenTangents.Small", new IMAGE_BRUSH("Icons/icon_CurveEditor_Straighten_40x", Icon20x20));
 	}
 
 	// PhysicsAssetEditor
diff --git a/Engine/Source/Editor/Matinee/Private/Matinee.cpp b/Engine/Source/Editor/Matinee/Private/Matinee.cpp
--- a/Engine/Source/Editor/Matinee/Private/Matinee.cpp
+++ b/Engine/Source/Editor/Matinee/Private/Matinee.cpp
@@ -3004,7 +3004,10 @@ void FMatinee::OnMenuCreateMovie()
 	// Ensure that this matinee is up and running before we start capturing
 	MovieSceneCapture->SetPrerequisiteActor(MatineeActor);
 
-	IMovieSceneCaptureDialogModule::Get().OpenDialog(LevelEditorModule.GetLevelEditorTabManager().ToSharedRef(), MovieSceneCapture);
+	// We initialize a default numeric interface as Matinee doesn't support the new mixed frame resolution concepts.
+	TSharedPtr<INumericTypeInterface<double>> NumericTypeInterface = MakeShareable(new TDefaultNumericTypeInterface<double>);
+
+	IMovieSceneCaptureDialogModule::Get().OpenDialog(LevelEditorModule.GetLevelEditorTabManager().ToSharedRef(), MovieSceneCapture, NumericTypeInterface);
 }
 
 #undef LOCTEXT_NAMESPACE
diff --git a/Engine/Source/Editor/MovieSceneCaptureDialog/MovieSceneCaptureDialog.Build.cs b/Engine/Source/Editor/MovieSceneCaptureDialog/MovieSceneCaptureDialog.Build.cs
--- a/Engine/Source/Editor/MovieSceneCaptureDialog/MovieSceneCaptureDialog.Build.cs
+++ b/Engine/Source/Editor/MovieSceneCaptureDialog/MovieSceneCaptureDialog.Build.cs
@@ -23,6 +23,7 @@ public class MovieSceneCaptureDialog : ModuleRules
 				"JsonUtilities",
 				"MovieScene",
 				"MovieSceneCapture",
+				"MovieSceneTools",
 				"PropertyEditor",
 				"SessionServices",
 				"Slate",
