Flatten works with weighted tangents. Since we don't have access to tick resolution on the section we do it as a two step, flatten as nonweighted then convert back to weighted.

diff --git a/Engine/Source/Editor/CurveEditor/Private/CurveEditor.cpp b/Engine/Source/Editor/CurveEditor/Private/CurveEditor.cpp
--- a/Engine/Source/Editor/CurveEditor/Private/CurveEditor.cpp
+++ b/Engine/Source/Editor/CurveEditor/Private/CurveEditor.cpp
@@ -593,7 +593,10 @@ void FCurveEditor::FlattenSelection()
 
 	TArray<FKeyHandle> KeyHandles;
 	TArray<FKeyAttributes> AllKeyPositions;
-
+	//Since we don't have access here to the Section to get Tick Resolution if we flatten a weighted tangent we
+	//do so by converting it to non-weighted and then back again.
+	TArray<FKeyHandle>  KeyHandlesWeighted;
+	TArray<FKeyAttributes> KeyAttributesWeighted;
 	for (const TTuple<FCurveModelID, FKeyHandleSet>& Pair : Selection.GetAll())
 	{
 		if (FCurveModel* Curve = FindCurve(Pair.Key))
@@ -604,6 +607,13 @@ void FCurveEditor::FlattenSelection()
 			AllKeyPositions.SetNum(KeyHandles.Num());
 			Curve->GetKeyAttributes(KeyHandles, AllKeyPositions);
 
+			KeyHandlesWeighted.Reset(Pair.Value.Num());
+			KeyHandlesWeighted.Append(Pair.Value.AsArray().GetData(), Pair.Value.Num());
+
+			KeyAttributesWeighted.SetNum(KeyHandlesWeighted.Num());
+			Curve->GetKeyAttributes(KeyHandlesWeighted, KeyAttributesWeighted);
+
+
 			// Straighten tangents, ignoring any keys that we can't set tangents on
 			for (int32 Index = AllKeyPositions.Num()-1 ; Index >= 0; --Index)
 			{
@@ -615,11 +625,31 @@ void FCurveEditor::FlattenSelection()
 					{
 						Attributes.SetTangentMode(RCTM_User);
 					}
+					//if any weighted convert and convert back to both (which is what only support other modes are not really used).,
+					if (Attributes.GetTangentWeightMode() == RCTWM_WeightedBoth || Attributes.GetTangentWeightMode() == RCTWM_WeightedArrive
+						|| Attributes.GetTangentWeightMode() == RCTWM_WeightedLeave)
+					{
+						Attributes.SetTangentWeightMode(RCTWM_WeightedNone);
+						FKeyAttributes& WeightedAttributes = KeyAttributesWeighted[Index];
+						WeightedAttributes.UnsetArriveTangent();
+						WeightedAttributes.UnsetLeaveTangent();
+						WeightedAttributes.UnsetArriveTangentWeight();
+						WeightedAttributes.UnsetLeaveTangentWeight();
+						WeightedAttributes.SetTangentWeightMode(RCTWM_WeightedBoth);
+
+					}
+					else
+					{
+						KeyAttributesWeighted.RemoveAtSwap(Index, 1, false);
+						KeyHandlesWeighted.RemoveAtSwap(Index, 1, false);
+					}
 				}
 				else
 				{
 					AllKeyPositions.RemoveAtSwap(Index, 1, false);
 					KeyHandles.RemoveAtSwap(Index, 1, false);
+					KeyAttributesWeighted.RemoveAtSwap(Index, 1, false);
+					KeyHandlesWeighted.RemoveAtSwap(Index, 1, false);
 				}
 			}
 
@@ -627,6 +657,10 @@ void FCurveEditor::FlattenSelection()
 			{
 				Curve->Modify();
 				Curve->SetKeyAttributes(KeyHandles, AllKeyPositions);
+				if (KeyAttributesWeighted.Num() > 0)
+				{
+					Curve->SetKeyAttributes(KeyHandlesWeighted, KeyAttributesWeighted);
+				}
 				bFoundAnyTangents = true;
 			}
 		}
