Curve Editor: Implemented basic text filtering for the curve editor tree

  - Filter terms are applied as a boolean OR, case-insensitive substring match
  - Pre-filtered expansion states are saved and restored when starting or clearing a filter
  - By default all matched items and parents are expanded, but not matched item's grand children. This gives the best tradeoff for filtered visibilty.
  - Tree items that are shown but did not match the filter terms are now shown with a subdued foreground color.

diff --git a/Engine/Source/Editor/CurveEditor/Private/CurveEditor.cpp b/Engine/Source/Editor/CurveEditor/Private/CurveEditor.cpp
--- a/Engine/Source/Editor/CurveEditor/Private/CurveEditor.cpp
+++ b/Engine/Source/Editor/CurveEditor/Private/CurveEditor.cpp
@@ -206,11 +206,6 @@ void FCurveEditor::RemoveTreeItem(FCurveEditorTreeItemID ItemID)
 	++ActiveCurvesSerialNumber;
 }
 
-FSimpleMulticastDelegate& FCurveEditor::OnTreeChanged()
-{
-	return Tree.OnChanged();
-}
-
 ECurveEditorTreeSelectionState FCurveEditor::GetTreeSelectionState(FCurveEditorTreeItemID InTreeItemID) const
 {
 	return Tree.GetSelectionState(InTreeItemID);
diff --git a/Engine/Source/Editor/CurveEditor/Private/Tree/CurveEditorTree.cpp b/Engine/Source/Editor/CurveEditor/Private/Tree/CurveEditorTree.cpp
--- a/Engine/Source/Editor/CurveEditor/Private/Tree/CurveEditorTree.cpp
+++ b/Engine/Source/Editor/CurveEditor/Private/Tree/CurveEditorTree.cpp
@@ -2,8 +2,11 @@
 
 #include "Tree/CurveEditorTree.h"
 #include "Tree/ICurveEditorTreeItem.h"
+#include "Tree/CurveEditorTreeFilter.h"
 #include "CurveEditor.h"
 
+#include "Algo/AnyOf.h"
+
 class FCurveModel;
 
 TArrayView<const FCurveModelID> FCurveEditorTreeItem::GetOrCreateCurves(FCurveEditor* CurveEditor)
@@ -49,9 +52,40 @@ void FCurveEditorTreeItem::DestroyUnpinnedCurves(FCurveEditor* CurveEditor)
 	}
 }
 
+FScopedCurveEditorTreeUpdateGuard::FScopedCurveEditorTreeUpdateGuard(FCurveEditorTree* InTree)
+	: Tree(InTree)
+{
+	Tree->OnChanged()->UpdateGuardCounter += 1;
+}
+
+FScopedCurveEditorTreeUpdateGuard::FScopedCurveEditorTreeUpdateGuard(FScopedCurveEditorTreeUpdateGuard&& RHS)
+	: Tree(RHS.Tree)
+{
+	Tree->OnChanged()->UpdateGuardCounter += 1;
+}
+
+FScopedCurveEditorTreeUpdateGuard& FScopedCurveEditorTreeUpdateGuard::operator=(FScopedCurveEditorTreeUpdateGuard&& RHS)
+{
+	if (&RHS != this)
+	{
+		Tree = RHS.Tree;
+		Tree->OnChanged()->UpdateGuardCounter += 1;
+	}
+	return *this;
+}
+
+FScopedCurveEditorTreeUpdateGuard::~FScopedCurveEditorTreeUpdateGuard()
+{
+	FCurveEditorOnChangedEvent* OnChangedEvent = Tree->OnChanged();
+	if (--OnChangedEvent->UpdateGuardCounter == 0)
+	{
+		TGuardValue<bool> Guard(OnChangedEvent->bIsBroadcastInProgress, true);
+		OnChangedEvent->Delegate.Broadcast();
+	}
+}
+
 FCurveEditorTree::FCurveEditorTree()
 {
-	bReentrantDelegate = false;
 	NextTreeItemID.Value = 1;
 }
 
@@ -75,10 +109,15 @@ const TArray<FCurveEditorTreeItemID>& FCurveEditorTree::GetRootItems() const
 	return RootItems.ChildIDs;
 }
 
+const TMap<FCurveEditorTreeItemID, FCurveEditorTreeItem>& FCurveEditorTree::GetAllItems() const
+{
+	return Items;
+}
+
 FCurveEditorTreeItem* FCurveEditorTree::AddItem(FCurveEditorTreeItemID ParentID)
 {
-	checkf(!bReentrantDelegate, TEXT("Curve editor tree must not be manipulated in response to it changing"));
-	TGuardValue<bool> Guard(bReentrantDelegate, true);
+	checkf(!OnChangedEvent.IsBroadcastInProgress(), TEXT("Curve editor tree must not be manipulated in response to it changing"));
+	FScopedCurveEditorTreeUpdateGuard BroadcastChangeUpdate(this);
 
 	FCurveEditorTreeItemID NewItemID = NextTreeItemID;
 
@@ -92,14 +131,13 @@ FCurveEditorTreeItem* FCurveEditorTree::AddItem(FCurveEditorTreeItemID ParentID)
 
 	++NextTreeItemID.Value;
 
-	OnChangedEvent.Broadcast();
 	return &NewItem;
 }
 
 void FCurveEditorTree::RemoveItem(FCurveEditorTreeItemID ItemID, FCurveEditor* CurveEditor)
 {
-	checkf(!bReentrantDelegate, TEXT("Curve editor tree must not be manipulated in response to it changing"));
-	TGuardValue<bool> Guard(bReentrantDelegate, true);
+	checkf(!OnChangedEvent.IsBroadcastInProgress(), TEXT("Curve editor tree must not be manipulated in response to it changing"));
+	FScopedCurveEditorTreeUpdateGuard BroadcastChangeUpdate(this);
 
 	FCurveEditorTreeItem* Item = Items.Find(ItemID);
 	if (!Item)
@@ -120,8 +158,6 @@ void FCurveEditorTree::RemoveItem(FCurveEditorTreeItemID ItemID, FCurveEditor* C
 
 	Items.Remove(ItemID);
 	Selection.Remove(ItemID);
-
-	OnChangedEvent.Broadcast();
 }
 
 void FCurveEditorTree::RemoveChildrenRecursive(TArray<FCurveEditorTreeItemID>&& LocalChildren, FCurveEditor* CurveEditor)
@@ -141,6 +177,108 @@ void FCurveEditorTree::RemoveChildrenRecursive(TArray<FCurveEditorTreeItemID>&&
 	}
 }
 
+
+bool FCurveEditorTree::FilterSpecificItems(TArrayView<FCurveEditorTreeFilter* const> FilterPtrs, TArrayView<const FCurveEditorTreeItemID> ItemsToFilter, ECurveEditorTreeFilterState InheritedState)
+{
+	bool bAnyMatched = false;
+
+	for (FCurveEditorTreeItemID ItemID : ItemsToFilter)
+	{
+		const FCurveEditorTreeItem& TreeItem = GetItem(ItemID);
+
+		// Retrieve the existing filter state for this item. This may have been set by preceeding filters.
+		ECurveEditorTreeFilterState FilterState         = InheritedState;
+		ECurveEditorTreeFilterState ChildInheritedState = InheritedState;
+
+		// If this has already not been matched, run it through our filter
+		TSharedPtr<ICurveEditorTreeItem> TreeItemImpl = TreeItem.GetItem();
+		if (TreeItemImpl)
+		{
+			const bool bMatchesFilter = Algo::AnyOf(FilterPtrs, [TreeItemImpl](const FCurveEditorTreeFilter* Filter){ return TreeItemImpl->PassesFilter(Filter); });
+			if (bMatchesFilter)
+			{
+				bAnyMatched = true;
+				FilterState = ECurveEditorTreeFilterState::Match;
+				ChildInheritedState = ECurveEditorTreeFilterState::ImplicitChild;
+			}
+		}
+
+		// Run the filter on all child nodes
+		const bool bMatchedChildren = FilterSpecificItems(FilterPtrs, TreeItem.GetChildren(), ChildInheritedState);
+
+		// If we matched children we become an implicit parent if not already matched
+		if (bMatchedChildren && FilterState != ECurveEditorTreeFilterState::Match)
+		{
+			bAnyMatched = true;
+			FilterState = ECurveEditorTreeFilterState::ImplicitParent;
+		}
+
+		if (FilterState != ECurveEditorTreeFilterState::NoMatch)
+		{
+			FilterStates.SetFilterState(ItemID, FilterState);
+		}
+	}
+
+	return bAnyMatched;
+}
+
+void FCurveEditorTree::RunFilters()
+{
+	checkf(!OnChangedEvent.IsBroadcastInProgress(), TEXT("Curve editor tree must not be manipulated in response to it changing"));
+
+	FScopedCurveEditorTreeUpdateGuard BroadcastChangeUpdate(this);
+
+	// Reset all the filter states back to the default
+	FilterStates.Reset();
+
+	if (Filters.Num())
+	{
+		FilterStates.Activate();
+
+		TArray<FCurveEditorTreeFilter*> FilterPtrs;
+
+		for (int32 Index = Filters.Num() - 1; Index >= 0; --Index)
+		{
+			TSharedPtr<FCurveEditorTreeFilter> Filter = Filters[Index].Pin();
+			if (!Filter)
+			{
+				// Remove invalid filters
+				Filters.RemoveAtSwap(Index, 1, false);
+			}
+			else
+			{
+				FilterPtrs.Add(Filter.Get());
+			}
+		}
+
+		FilterSpecificItems(FilterPtrs, RootItems.ChildIDs, ECurveEditorTreeFilterState::NoMatch);
+	}
+	else
+	{
+		FilterStates.Deactivate();
+	}
+}
+
+void FCurveEditorTree::AddFilter(TWeakPtr<FCurveEditorTreeFilter> NewFilter)
+{
+	Filters.AddUnique(NewFilter);
+}
+
+void FCurveEditorTree::RemoveFilter(TWeakPtr<FCurveEditorTreeFilter> FilterToRemove)
+{
+	Filters.Remove(FilterToRemove);
+}
+
+const FCurveEditorFilterStates& FCurveEditorTree::GetFilterStates() const
+{
+	return FilterStates;
+}
+
+ECurveEditorTreeFilterState FCurveEditorTree::GetFilterState(FCurveEditorTreeItemID InTreeItemID) const
+{
+	return FilterStates.Get(InTreeItemID);
+}
+
 void FCurveEditorTree::SetDirectSelection(TArray<FCurveEditorTreeItemID>&& TreeItems)
 {
 	Selection.Reset();
@@ -176,4 +314,4 @@ ECurveEditorTreeSelectionState FCurveEditorTree::GetSelectionState(FCurveEditorT
 {
 	const ECurveEditorTreeSelectionState* State = Selection.Find(InTreeItemID);
 	return State ? *State : ECurveEditorTreeSelectionState::None;
-}
\ No newline at end of file
+}
diff --git a/Engine/Source/Editor/CurveEditor/Private/Tree/CurveEditorTreeFilter.cpp b/Engine/Source/Editor/CurveEditor/Private/Tree/CurveEditorTreeFilter.cpp
new file mode 100644
--- /dev/null
+++ b/Engine/Source/Editor/CurveEditor/Private/Tree/CurveEditorTreeFilter.cpp
@@ -0,0 +1,24 @@
+// Copyright 1998-2019 Epic Games, Inc. All Rights Reserved.
+
+#include "Tree/CurveEditorTreeFilter.h"
+#include "Tree/ICurveEditorTreeItem.h"
+#include "Tree/CurveEditorTree.h"
+
+
+
+ECurveEditorTreeFilterType FCurveEditorTreeFilter::RegisterFilterType()
+{
+	static ECurveEditorTreeFilterType NextFilterType = ECurveEditorTreeFilterType::CUSTOM_START;
+	ensureMsgf(NextFilterType != ECurveEditorTreeFilterType::First, TEXT("Maximum limit for registered curve tree filters (64) reached."));
+	if (NextFilterType == ECurveEditorTreeFilterType::First)
+	{
+		return NextFilterType;
+	}
+
+	ECurveEditorTreeFilterType ThisFilterType = NextFilterType;
+
+	// When the custom view ID reaches 0x80000000 the left shift will result in well-defined unsigned integer wraparound, resulting in 0 (None)
+	NextFilterType = ECurveEditorTreeFilterType( ((__underlying_type(ECurveEditorTreeFilterType))NextFilterType) + 1 );
+
+	return NextFilterType;
+}
\ No newline at end of file
diff --git a/Engine/Source/Editor/CurveEditor/Private/Tree/SCurveEditorTree.cpp b/Engine/Source/Editor/CurveEditor/Private/Tree/SCurveEditorTree.cpp
--- a/Engine/Source/Editor/CurveEditor/Private/Tree/SCurveEditorTree.cpp
+++ b/Engine/Source/Editor/CurveEditor/Private/Tree/SCurveEditorTree.cpp
@@ -76,6 +76,16 @@ struct SCurveEditorTableRow : SMultiColumnTableRow<FCurveEditorTreeItemID>
 		WeakCurveEditor = InCurveEditor;
 
 		SMultiColumnTableRow::Construct(InArgs, OwnerTableView);
+
+		SetForegroundColor(MakeAttributeSP(this, &SCurveEditorTableRow::GetForegroundColorByFilterState));
+	}
+
+	FSlateColor GetForegroundColorByFilterState() const
+	{
+		TSharedPtr<FCurveEditor> CurveEditor = WeakCurveEditor.Pin();
+
+		const bool bIsMatch = CurveEditor.IsValid() && ( CurveEditor->GetTree()->GetFilterState(TreeItemID) == ECurveEditorTreeFilterState::Match );
+		return bIsMatch ? FSlateColor::UseForeground() : FSlateColor::UseSubduedForeground();
 	}
 
 	virtual TSharedRef<SWidget> GenerateWidgetForColumn(const FName& InColumnName) override
@@ -117,6 +127,7 @@ struct SCurveEditorTableRow : SMultiColumnTableRow<FCurveEditorTreeItemID>
 
 void SCurveEditorTree::Construct(const FArguments& InArgs, TSharedPtr<FCurveEditor> InCurveEditor)
 {
+	bFilterWasActive = false;
 	HeaderRow = SNew(SHeaderRow)
 		.Visibility(EVisibility::Collapsed)
 
@@ -133,7 +144,7 @@ void SCurveEditorTree::Construct(const FArguments& InArgs, TSharedPtr<FCurveEdit
 		.SelectionMode(ESelectionMode::Multi)
 		.HeaderRow(HeaderRow)
 		.HighlightParentNodesForSelection(true)
-		.TreeItemsSource(&InCurveEditor->GetRootTreeItems())
+		.TreeItemsSource(&RootItems)
 		.OnGetChildren(this, &SCurveEditorTree::GetTreeItemChildren)
 		.OnGenerateRow(this, &SCurveEditorTree::GenerateRow)
 		.OnSetExpansionRecursive(this, &SCurveEditorTree::SetItemExpansionRecursive)
@@ -145,7 +156,63 @@ void SCurveEditorTree::Construct(const FArguments& InArgs, TSharedPtr<FCurveEdit
 		)
 	];
 
-	CurveEditor->OnTreeChanged().AddSP(TreeView.ToSharedRef(), &SCurveEditorTreeView::RequestTreeRefresh);
+	CurveEditor->GetTree()->Bind_OnChanged(FSimpleDelegate::CreateSP(this, &SCurveEditorTree::RefreshTree));
+}
+
+void SCurveEditorTree::RefreshTree()
+{
+	RootItems.Reset();
+
+	const FCurveEditorTree* CurveEditorTree = CurveEditor->GetTree();
+	const FCurveEditorFilterStates& FilterStates = CurveEditorTree->GetFilterStates();
+
+	// When changing to/from a filtered state, we save and restore expansion states
+	if (FilterStates.IsActive() && !bFilterWasActive)
+	{
+		// Save expansion states
+		PreFilterExpandedItems.Reset();
+		TreeView->GetExpandedItems(PreFilterExpandedItems);
+	}
+	else if (!FilterStates.IsActive() && bFilterWasActive)
+	{
+		// Restore expansion states
+		TreeView->ClearExpandedItems();
+		for (FCurveEditorTreeItemID ExpandedItem : PreFilterExpandedItems)
+		{
+			TreeView->SetItemExpansion(ExpandedItem, true);
+		}
+		PreFilterExpandedItems.Reset();
+	}
+
+	// Repopulate root tree items based on filters
+	for (FCurveEditorTreeItemID RootItemID : CurveEditor->GetRootTreeItems())
+	{
+		if (FilterStates.Get(RootItemID) != ECurveEditorTreeFilterState::NoMatch)
+		{
+			RootItems.Add(RootItemID);
+		}
+	}
+
+	RootItems.Shrink();
+	TreeView->RequestTreeRefresh();
+
+	if (FilterStates.IsActive())
+	{
+		// If a filter is active, all matched items and their parents are expanded
+		TreeView->ClearExpandedItems();
+		for (const TTuple<FCurveEditorTreeItemID, FCurveEditorTreeItem>& Pair : CurveEditorTree->GetAllItems())
+		{
+			ECurveEditorTreeFilterState FilterState = FilterStates.Get(Pair.Key);
+
+			// Expand any matched items or parents of matched items
+			if (FilterState == ECurveEditorTreeFilterState::Match || FilterState == ECurveEditorTreeFilterState::ImplicitParent)
+			{
+				TreeView->SetItemExpansion(Pair.Key, true);
+			}
+		}
+	}
+
+	bFilterWasActive = FilterStates.IsActive();
 }
 
 FReply SCurveEditorTree::OnKeyDown(const FGeometry& MyGeometry, const FKeyEvent& InKeyEvent)
@@ -166,8 +233,15 @@ TSharedRef<ITableRow> SCurveEditorTree::GenerateRow(FCurveEditorTreeItemID ItemI
 
 void SCurveEditorTree::GetTreeItemChildren(FCurveEditorTreeItemID Parent, TArray<FCurveEditorTreeItemID>& OutChildren)
 {
-	TArrayView<const FCurveEditorTreeItemID> Children = CurveEditor->GetTreeItem(Parent).GetChildren();
-	OutChildren.Append(Children.GetData(), Children.Num());
+	const FCurveEditorFilterStates& FilterStates = CurveEditor->GetTree()->GetFilterStates();
+
+	for (FCurveEditorTreeItemID ChildID : CurveEditor->GetTreeItem(Parent).GetChildren())
+	{
+		if (FilterStates.Get(ChildID) != ECurveEditorTreeFilterState::NoMatch)
+		{
+			OutChildren.Add(ChildID);
+		}
+	}
 }
 
 void SCurveEditorTree::OnTreeSelectionChanged(FCurveEditorTreeItemID, ESelectInfo::Type)
diff --git a/Engine/Source/Editor/CurveEditor/Private/Tree/SCurveEditorTreeTextFilter.cpp b/Engine/Source/Editor/CurveEditor/Private/Tree/SCurveEditorTreeTextFilter.cpp
new file mode 100644
--- /dev/null
+++ b/Engine/Source/Editor/CurveEditor/Private/Tree/SCurveEditorTreeTextFilter.cpp
@@ -0,0 +1,49 @@
+// Copyright 1998-2019 Epic Games, Inc. All Rights Reserved.
+
+#include "Tree/SCurveEditorTreeTextFilter.h"
+#include "Tree/CurveEditorTreeFilter.h"
+#include "Tree/CurveEditorTree.h"
+#include "CurveEditor.h"
+
+#include "Widgets/Input/SSearchBox.h"
+
+
+void SCurveEditorTreeTextFilter::Construct(const FArguments& InArgs, TSharedPtr<FCurveEditor> CurveEditor)
+{
+	WeakCurveEditor = CurveEditor;
+
+	ChildSlot
+	[
+		SNew(SSearchBox)
+		.HintText(NSLOCTEXT("CurveEditor", "TextFilterHint", "Filter"))
+		.OnTextChanged(this, &SCurveEditorTreeTextFilter::OnFilterTextChanged)
+	];
+}
+
+void SCurveEditorTreeTextFilter::OnFilterTextChanged( const FText& FilterText )
+{
+	TSharedPtr<FCurveEditor> CurveEditor = WeakCurveEditor.Pin();
+	if (CurveEditor)
+	{
+		if (!Filter)
+		{
+			Filter = MakeShared<FCurveEditorTreeTextFilter>();
+		}
+
+		Filter->FilterTerms.Reset();
+
+		static const bool bCullEmpty = true;
+		FilterText.ToString().ParseIntoArray(Filter->FilterTerms, TEXT(" "), bCullEmpty);
+
+		if (Filter->FilterTerms.Num() > 0)
+		{
+			CurveEditor->GetTree()->AddFilter(Filter);
+		}
+		else
+		{
+			CurveEditor->GetTree()->RemoveFilter(Filter);
+		}
+
+		CurveEditor->GetTree()->RunFilters();
+	}
+}
\ No newline at end of file
diff --git a/Engine/Source/Editor/CurveEditor/Public/CurveEditor.h b/Engine/Source/Editor/CurveEditor/Public/CurveEditor.h
--- a/Engine/Source/Editor/CurveEditor/Public/CurveEditor.h
+++ b/Engine/Source/Editor/CurveEditor/Public/CurveEditor.h
@@ -293,11 +293,6 @@ public:
 	 */
 	void RemoveTreeItem(FCurveEditorTreeItemID ItemID);
 
-	/**
-	 * Retrieve an event that is broadcast any time the tree changes
-	 */
-	FSimpleMulticastDelegate& OnTreeChanged();
-
 	/**
 	 * Set the current tree selection
 	 */
@@ -313,6 +308,14 @@ public:
 	 */
 	const TMap<FCurveEditorTreeItemID, ECurveEditorTreeSelectionState>& GetTreeSelection() const;
 
+	/**
+	 * Access the curve editor tree.
+	 */
+	FCurveEditorTree* GetTree()
+	{
+		return &Tree;
+	}
+
 	/**
 	 * Retrieve a serial number that is incremented any time a curve is added or removed
 	 */
diff --git a/Engine/Source/Editor/CurveEditor/Public/Tree/CurveEditorTree.h b/Engine/Source/Editor/CurveEditor/Public/Tree/CurveEditorTree.h
--- a/Engine/Source/Editor/CurveEditor/Public/Tree/CurveEditorTree.h
+++ b/Engine/Source/Editor/CurveEditor/Public/Tree/CurveEditorTree.h
@@ -11,10 +11,78 @@
 #include "CurveEditorTypes.h"
 
 struct ICurveEditorTreeItem;
+struct FCurveEditorTreeFilter;
 
 class FCurveEditor;
 class FCurveEditorTree;
 
+/** Enumeration specifying how a specific tree item has matched the current set of filters */
+enum class ECurveEditorTreeFilterState : uint8
+{
+	/** The item did not match any filter, and neither did any of its parents or children */
+	NoMatch,
+
+	/** Neither this item nor any of its children match filters, but one of its parents did (ie it resides within a matched item) */
+	ImplicitChild,
+
+	/** Neither this item nor any of its parents match the filters, but one of its descendant children did (ie it is a parent of a matched item) */
+	ImplicitParent,
+
+	/** This item itself matched one or more of the filters */
+	Match,
+};
+
+/**
+ * Scoped guard that will trigger the tree OnChanged event when all scoped guards have been exited
+ */
+struct CURVEEDITOR_API FScopedCurveEditorTreeUpdateGuard
+{
+	explicit FScopedCurveEditorTreeUpdateGuard(FCurveEditorTree* InTree);
+	~FScopedCurveEditorTreeUpdateGuard();
+
+	FScopedCurveEditorTreeUpdateGuard(FScopedCurveEditorTreeUpdateGuard&& RHS);
+	FScopedCurveEditorTreeUpdateGuard& operator=(FScopedCurveEditorTreeUpdateGuard&& RHS);
+
+private:
+	FCurveEditorTree* Tree;
+};
+
+/**
+ * Struct that represents an event for when the tree has been changed.
+ * This type carefully only allows FScopedCurveEditorTreeUpdateGuard to broadcast the event, and makes special checks for re-entrancy
+ */
+struct FCurveEditorOnChangedEvent
+{
+	/**
+	 * @return true if the event is currently being broadcast
+	 */
+	bool IsBroadcastInProgress() const
+	{
+		return bIsBroadcastInProgress;
+	}
+
+	/**
+	 * Bind a new handler to this OnChanged event
+	 */
+	FDelegateHandle Bind(FSimpleDelegate&& Handler) { return Delegate.Add(MoveTemp(Handler)); }
+
+	/**
+	 * Unbind a previously bound handler from this event
+	 */
+	void Unbind(FDelegateHandle Handle) { return Delegate.Remove(Handle); }
+
+private:
+
+	friend FScopedCurveEditorTreeUpdateGuard;
+
+	/** Safety check to ensure that invocations of OnChangedEvent are never re-entrant */
+	bool bIsBroadcastInProgress = false;
+	/** Counter that is incremented for each living instance of FCurveEditorOnChangedEvent */
+	uint32 UpdateGuardCounter = 0;
+	/** The actual multi-cast delegate */
+	FSimpleMulticastDelegate Delegate;
+};
+
 /**
  * Container specifying a linear set of child identifiers and 
  */
@@ -124,6 +192,76 @@ private:
 };
 
 
+/**
+ * Sparse map of filter states specifying items that have matched a filter
+ */
+struct FCurveEditorFilterStates
+{
+	/**
+	 * Reset all the filter states currently being tracked (does not affect IsActive
+	 */
+	void Reset()
+	{
+		FilterStates.Reset();
+	}
+
+	/**
+	 * Retrieve the filter state for a specific tree item ID
+	 * @return The item's filter state, or ECurveEditorTreeFilterState::Match if filters are not currently active.
+	 */
+	ECurveEditorTreeFilterState Get(FCurveEditorTreeItemID ItemID) const
+	{
+		if (!bIsActive)
+		{
+			// If not active, everything is treated as having matched the (non-existent) filters
+			return ECurveEditorTreeFilterState::Match;
+		}
+
+		const ECurveEditorTreeFilterState* State = FilterStates.Find(ItemID);
+		return State ? *State : ECurveEditorTreeFilterState::NoMatch;
+	}
+
+	/**
+	 * Assign a new filter state to an item
+	 */
+	void SetFilterState(FCurveEditorTreeItemID ItemID, ECurveEditorTreeFilterState NewState)
+	{
+		FilterStates.Add(ItemID, NewState);
+	}
+
+	/**
+	 * Check whether filters are active or not
+	 */
+	bool IsActive() const
+	{
+		return bIsActive;
+	}
+
+	/**
+	 * Activate the filters so that they begin to take effect
+	 */
+	void Activate()
+	{
+		bIsActive = true;
+	}
+
+	/**
+	 * Deactivate the filters so that they no longer take effect
+	 */
+	void Deactivate()
+	{
+		bIsActive = false;
+	}
+
+private:
+
+	/** Whether filters should be active or not */
+	bool bIsActive = false;
+
+	/** Filter state map. Items with no implicit or explicit filter state are not present */
+	TMap<FCurveEditorTreeItemID, ECurveEditorTreeFilterState> FilterStates;
+};
+
 /** 
  * Complete implementation of a curve editor tree. Only really defines the hierarchy and selection states for tree items.
  */
@@ -149,10 +287,15 @@ public:
 	FCurveEditorTreeItem* FindItem(FCurveEditorTreeItemID ItemID);
 
 	/** 
-	 * Retrieve this curve editor's root items
+	 * Retrieve this curve editor's root items irrespective of filter state
 	 */
 	const TArray<FCurveEditorTreeItemID>& GetRootItems() const;
 
+	/** 
+	 * Retrieve all the items stored in this tree irrespective of filter state
+	 */
+	const TMap<FCurveEditorTreeItemID, FCurveEditorTreeItem>& GetAllItems() const;
+
 	/**
 	 * Add a new empty item to the tree
 	 *
@@ -169,12 +312,23 @@ public:
 	void RemoveItem(FCurveEditorTreeItemID ItemID, FCurveEditor* CurveEditor);
 
 	/**
-	 * Access a delegate that is executed whenever this tree's hiararchy has changed in some way.
+	 * Run all the filters on this tree, updating filter state for all tree items
 	 */
-	FSimpleMulticastDelegate& OnChanged()
-	{
-		return OnChangedEvent;
-	}
+	void RunFilters();
+
+	/**
+	 * Add a new filter to this tree. Does not run the filter (and thus update any tree views) until RunFilters is called.
+	 *
+	 * @param NewFilter The new filter to add to this tree
+	 */
+	void AddFilter(TWeakPtr<FCurveEditorTreeFilter> NewFilter);
+
+	/**
+	 * Remove an existing filter from this tree. Does not re-run the filters (and thus update any tree views) until RunFilters is called.
+	 *
+	 * @param FilterToRemove The filter to remove from this tree
+	 */
+	void RemoveFilter(TWeakPtr<FCurveEditorTreeFilter> FilterToRemove);
 
 	/**
 	 * Inform this tree that the specified tree item IDs have been directly selected on the UI.
@@ -192,26 +346,83 @@ public:
 	 */
 	ECurveEditorTreeSelectionState GetSelectionState(FCurveEditorTreeItemID InTreeItemID) const;
 
+	/**
+	 * Access the filter state for this tree. Items that are neither implicitly or explicitly filtered-in are not present in the map.
+	 */
+	const FCurveEditorFilterStates& GetFilterStates() const;
+
+	/**
+	 * Check a specific tree item's filter state
+	 */
+	ECurveEditorTreeFilterState GetFilterState(FCurveEditorTreeItemID InTreeItemID) const;
+
+	/**
+	 * Retrieve this tree's on-changed event
+	 */
+	FCurveEditorOnChangedEvent* OnChanged()
+	{
+		return &OnChangedEvent;
+	}
+
+	/**
+	 * Retrieve a scoped guard that will broadcast the on changed handlers for this tree when it goes out of scope (along with all other scoped guards on the stack)
+	 * Can be used to defer such broadcasts in situations where many changes are made to the tree at a time.
+	 */
+	FScopedCurveEditorTreeUpdateGuard ScopedUpdateGuard()
+	{
+		return FScopedCurveEditorTreeUpdateGuard(this);
+	}
+
+	/**
+	 * Add a handler for when this tree structure is changed in some way (items added/removed, tree filters changed etc)
+	 */
+	FDelegateHandle Bind_OnChanged(FSimpleDelegate&& Handler)
+	{
+		return OnChangedEvent.Bind(MoveTemp(Handler));
+	}
+
+	/**
+	 * Remove a handler for when this tree structure is changed in some way (items added/removed, tree filters changed etc)
+	 */
+	void Unbind_OnChanged(FDelegateHandle Handle)
+	{
+		return OnChangedEvent.Unbind(Handle);
+	}
+
 private:
 
 	// Recursively removes children without removing them from the parent (assuming the parent is also being removed)
 	void RemoveChildrenRecursive(TArray<FCurveEditorTreeItemID>&& Children, FCurveEditor* CurveEditor);
 
-	/** Safety check to ensure that invocations of OnChangedEvent are never re-entrant */
-	bool bReentrantDelegate;
+	/**
+	 * Run the specified filters over the specified items and their recursive children, storing the results in this instance's FilterStates struct.
+	 *
+	 * @param FilterPtrs     Array of non-null pointers to filters to use. Items are considered matched if they match any filter in this array.
+	 * @param Items          Array item IDs to filter
+	 * @param InheritedState The filter state for each item to receive if it does not directly match a filter (either ECurveEditorTreeFilterState::NoMatch or ECurveEditorTreeFilterState::InheritedChild)
+	 * @return Whether any of the items or any their recursive children matched any filter
+	 */
+	bool FilterSpecificItems(TArrayView<FCurveEditorTreeFilter* const> FilterPtrs, TArrayView<const FCurveEditorTreeItemID> Items, ECurveEditorTreeFilterState InheritedState);
 
-	/** Incrememnting ID for the next tree item to be created */
+	/** Incrementing ID for the next tree item to be created */
 	FCurveEditorTreeItemID NextTreeItemID;
 
-	FSimpleMulticastDelegate OnChangedEvent;
+	/** Container housing the machinery required for deferred broadcast of changes to the tree */
+	FCurveEditorOnChangedEvent OnChangedEvent;
 
 	/** Map of all tree items by their ID */
 	TMap<FCurveEditorTreeItemID, FCurveEditorTreeItem> Items;
 
+	/** Map of all tree items by their ID */
+	TArray<TWeakPtr<FCurveEditorTreeFilter>> Filters;
+
 	/** Hierarchical information for the tree */
 	FSortedCurveEditorTreeItems RootItems;
 	TMap<FCurveEditorTreeItemID, FSortedCurveEditorTreeItems> ChildItemIDs;
 
 	/** Selection state map. Items with no implicit or explicit selection are not present */
 	TMap<FCurveEditorTreeItemID, ECurveEditorTreeSelectionState> Selection;
+
+	/** Filter state map. Items with no implicit or explicit filter state are not present */
+	FCurveEditorFilterStates FilterStates;
 };
\ No newline at end of file
diff --git a/Engine/Source/Editor/CurveEditor/Public/Tree/CurveEditorTreeFilter.h b/Engine/Source/Editor/CurveEditor/Public/Tree/CurveEditorTreeFilter.h
new file mode 100644
--- /dev/null
+++ b/Engine/Source/Editor/CurveEditor/Public/Tree/CurveEditorTreeFilter.h
@@ -0,0 +1,87 @@
+// Copyright 1998-2019 Epic Games, Inc. All Rights Reserved.
+
+#pragma once
+
+#include "Templates/SharedPointer.h"
+#include "Misc/EnumClassFlags.h"
+#include "Containers/UnrealString.h"
+
+class FCurveEditorTree;
+struct FCurveEditorTreeItem;
+struct FCurveEditorTreeItemID;
+struct FCurveEditorFilterStates;
+template<typename> class TArrayView;
+enum class ECurveEditorTreeFilterState : uint8;
+
+enum class ECurveEditorTreeFilterType : uint32
+{
+	/** Filter is a FCurveEditorTreeTextFilter instance */
+	Text,
+
+	CUSTOM_START,
+
+
+	First = Text,
+};
+
+
+
+struct CURVEEDITOR_API FCurveEditorTreeFilter
+{
+	explicit FCurveEditorTreeFilter(ECurveEditorTreeFilterType InFilterType)
+		: FilterType(InFilterType)
+	{}
+
+	virtual ~FCurveEditorTreeFilter() {}
+
+	/**
+	 * @return The type of this filter as registered by RegisterFilterType
+	 */
+	ECurveEditorTreeFilterType GetType() const
+	{
+		return FilterType;
+	}
+
+public:
+
+	/**
+	 * Register a new filter type that is passed to ICurveEditorTreeItem::Filter
+	 *
+	 * @param InTreeItem The tree item to test
+	 * @return true if the item matches the filter, false otherwise
+	 */
+	static ECurveEditorTreeFilterType RegisterFilterType();
+
+protected:
+
+	/** The static type of this filter as retrieved by RegisterFilterType */
+	ECurveEditorTreeFilterType FilterType;
+};
+
+/**
+ * Built-in text filter of type ECurveEditorTreeFilterType::Text. Filter terms are applied as a case-insensitive boolean OR substring match.
+ */
+struct FCurveEditorTreeTextFilter : FCurveEditorTreeFilter
+{
+	/** Array of case-insensitive terms to find within tree items. */
+	TArray<FString> FilterTerms;
+
+	FCurveEditorTreeTextFilter()
+		: FCurveEditorTreeFilter(ECurveEditorTreeFilterType::Text)
+	{}
+
+	/**
+	 * Check whether the supplied string matches any of the terms
+	 */
+	bool Match(const TCHAR* InString) const
+	{
+		for (const FString& Term : FilterTerms)
+		{
+			if (FCString::Stristr(InString, *Term) != nullptr)
+			{
+				return true;
+			}
+		}
+		return false;
+	}
+};
\ No newline at end of file
diff --git a/Engine/Source/Editor/CurveEditor/Public/Tree/CurveEditorTreeTraits.h b/Engine/Source/Editor/CurveEditor/Public/Tree/CurveEditorTreeTraits.h
new file mode 100644
--- /dev/null
+++ b/Engine/Source/Editor/CurveEditor/Public/Tree/CurveEditorTreeTraits.h
@@ -0,0 +1,60 @@
+// Copyright 1998-2019 Epic Games, Inc. All Rights Reserved.
+
+#pragma once
+
+#include "Widgets/DeclarativeSyntaxSupport.h"
+#include "CurveEditorTypes.h"
+
+template<>
+struct TListTypeTraits<FCurveEditorTreeItemID>
+{
+public:
+	struct NullableType : FCurveEditorTreeItemID
+	{
+		NullableType(TYPE_OF_NULLPTR){}
+		NullableType(FCurveEditorTreeItemID Other) : FCurveEditorTreeItemID(Other) {}
+	};
+
+	using MapKeyFuncs = TDefaultMapHashableKeyFuncs<FCurveEditorTreeItemID, TSharedRef<ITableRow>, false>;
+	using MapKeyFuncsSparse = TDefaultMapHashableKeyFuncs<FCurveEditorTreeItemID, FSparseItemInfo, false>;
+	using SetKeyFuncs = DefaultKeyFuncs<FCurveEditorTreeItemID>;
+
+	static void AddReferencedObjects(FReferenceCollector&, TArray<FCurveEditorTreeItemID>&, TSet<FCurveEditorTreeItemID>&) {}
+
+	static bool IsPtrValid(NullableType InPtr)
+	{
+		return InPtr.IsValid();
+	}
+
+	static void ResetPtr(NullableType& InPtr)
+	{
+		InPtr = nullptr;
+	}
+
+	static NullableType MakeNullPtr()
+	{
+		return nullptr;
+	}
+
+	static FCurveEditorTreeItemID NullableItemTypeConvertToItemType(NullableType InPtr)
+	{
+		return InPtr;
+	}
+
+	static FString DebugDump(FCurveEditorTreeItemID InPtr)
+	{
+		return FString::Printf(TEXT("%d"), InPtr.GetValue());
+	}
+
+	class SerializerType{};
+};
+
+template <>
+struct TIsValidListItem<FCurveEditorTreeItemID>
+{
+	enum
+	{
+		Value = true
+	};
+};
+
diff --git a/Engine/Source/Editor/CurveEditor/Public/Tree/ICurveEditorTreeItem.h b/Engine/Source/Editor/CurveEditor/Public/Tree/ICurveEditorTreeItem.h
--- a/Engine/Source/Editor/CurveEditor/Public/Tree/ICurveEditorTreeItem.h
+++ b/Engine/Source/Editor/CurveEditor/Public/Tree/ICurveEditorTreeItem.h
@@ -6,27 +6,58 @@
 #include "Containers/Array.h"
 #include "Templates/UniquePtr.h"
 
+enum class ECurveEditorTreeFilterType : uint32;
+
 struct FCurveEditorTreeItemID;
+struct FCurveEditorTreeFilter;
 
 class FName;
 class SWidget;
 class FCurveModel;
 class FCurveEditor;
 
+/**
+ * Optional implementation interface for any tree item to be shown on the curve editor tree.
+ */
 struct CURVEEDITOR_API ICurveEditorTreeItem
 {
 	virtual ~ICurveEditorTreeItem() {}
+
+	/**
+	 * Generate the widget content for the specified column name of the curve editor tree view
+	 *
+	 * @param InColumnName     The name of the column to generate widget content for. See FColumnNames for valid names.
+	 * @param InCurveEditor    Weak pointer to the curve editor instance. Persistent TSharedPtrs should not be held to this pointer.
+	 * @param InTreeItemID     The ID of the tree item that this interface is assigned to
+	 * @return (Optional) Widget content for this column, or nullptr if none is necessary for this column.
+	 */
 	virtual TSharedPtr<SWidget> GenerateCurveEditorTreeWidget(const FName& InColumnName, TWeakPtr<FCurveEditor> InCurveEditor, FCurveEditorTreeItemID InTreeItemID) = 0;
 
+	/**
+	 * Populate the specified array with curve models that are represented by this tree item
+	 *
+	 * @param OutCurveModels  Array of curve models to be populated with the curves of this tree item
+	 */
 	virtual void CreateCurveModels(TArray<TUniquePtr<FCurveModel>>& OutCurveModels) = 0;
 
+	/**
+	 * Check whether this tree item passes the specified filter. Filter's types can be checked using FCurveEditorTreeFilter::GetType() and static_casted for known filter types.
+	 *
+	 * @param InFilter        The filter to test against.
+	 * @return True if this item is deemed to have passed the filter, false otherwise
+	 */
+	virtual bool PassesFilter(const FCurveEditorTreeFilter* InFilter) const { return false; }
+
+	/** Structure of pre-defined supported column names for the curve editor outliner */
 	struct FColumnNames
 	{
+		/** Generic label content */
 		FName Label;
+		/** Right-aligned column for pin buttons */
 		FName PinHeader;
 
 		CURVEEDITOR_API FColumnNames();
 	};
 
 	static const FColumnNames ColumnNames;
-};
\ No newline at end of file
+};
diff --git a/Engine/Source/Editor/CurveEditor/Public/Tree/SCurveEditorTree.h b/Engine/Source/Editor/CurveEditor/Public/Tree/SCurveEditorTree.h
--- a/Engine/Source/Editor/CurveEditor/Public/Tree/SCurveEditorTree.h
+++ b/Engine/Source/Editor/CurveEditor/Public/Tree/SCurveEditorTree.h
@@ -35,8 +35,17 @@ private:
 
 	void SetItemExpansionRecursive(FCurveEditorTreeItemID Model, bool bInExpansionState);
 
+	void RefreshTree();
+
 private:
 
+	bool bFilterWasActive;
+
+	TArray<FCurveEditorTreeItemID> RootItems;
+
+	/** Set of item IDs that were expanded before a filter was applied */
+	TSet<FCurveEditorTreeItemID> PreFilterExpandedItems;
+
 	TSharedPtr<FCurveEditor> CurveEditor;
 
 	TSharedPtr<SHeaderRow> HeaderRow;
diff --git a/Engine/Source/Editor/CurveEditor/Public/Tree/SCurveEditorTreeTextFilter.h b/Engine/Source/Editor/CurveEditor/Public/Tree/SCurveEditorTreeTextFilter.h
new file mode 100644
--- /dev/null
+++ b/Engine/Source/Editor/CurveEditor/Public/Tree/SCurveEditorTreeTextFilter.h
@@ -0,0 +1,26 @@
+// Copyright 1998-2019 Epic Games, Inc. All Rights Reserved.
+
+#pragma once
+
+#include "Widgets/DeclarativeSyntaxSupport.h"
+#include "Widgets/SCompoundWidget.h"
+
+class FCurveEditor;
+struct FCurveEditorTreeTextFilter;
+
+class CURVEEDITOR_API SCurveEditorTreeTextFilter : public SCompoundWidget
+{
+public:
+
+	SLATE_BEGIN_ARGS(SCurveEditorTreeTextFilter){}
+	SLATE_END_ARGS()
+
+	void Construct(const FArguments& InArgs, TSharedPtr<FCurveEditor> CurveEditor);
+
+private:
+
+	void OnFilterTextChanged(const FText& FilterText);
+
+	TSharedPtr<FCurveEditorTreeTextFilter> Filter;
+	TWeakPtr<FCurveEditor> WeakCurveEditor;
+};
\ No newline at end of file
diff --git a/Engine/Source/Editor/Sequencer/Private/DisplayNodes/SequencerDisplayNode.cpp b/Engine/Source/Editor/Sequencer/Private/DisplayNodes/SequencerDisplayNode.cpp
--- a/Engine/Source/Editor/Sequencer/Private/DisplayNodes/SequencerDisplayNode.cpp
+++ b/Engine/Source/Editor/Sequencer/Private/DisplayNodes/SequencerDisplayNode.cpp
@@ -26,6 +26,7 @@
 #include "CommonMovieSceneTools.h"
 #include "Framework/Commands/GenericCommands.h"
 #include "Tree/SCurveEditorTreePin.h"
+#include "Tree/CurveEditorTreeFilter.h"
 #include "ScopedTransaction.h"
 #include "SequencerKeyTimeCache.h"
 #include "SequencerNodeSortingMethods.h"
@@ -1054,5 +1055,14 @@ void FSequencerDisplayNode::CreateCurveModels(TArray<TUniquePtr<FCurveModel>>& O
 {
 }
 
+bool FSequencerDisplayNode::PassesFilter(const FCurveEditorTreeFilter* InFilter) const
+{
+	if (InFilter->GetType() == ECurveEditorTreeFilterType::Text)
+	{
+		const FCurveEditorTreeTextFilter* Filter = static_cast<const FCurveEditorTreeTextFilter*>(InFilter);
+		return Filter->Match(*GetDisplayName().ToString());
+	}
+	return false;
+}
 
 #undef LOCTEXT_NAMESPACE
diff --git a/Engine/Source/Editor/Sequencer/Private/DisplayNodes/SequencerDisplayNode.h b/Engine/Source/Editor/Sequencer/Private/DisplayNodes/SequencerDisplayNode.h
--- a/Engine/Source/Editor/Sequencer/Private/DisplayNodes/SequencerDisplayNode.h
+++ b/Engine/Source/Editor/Sequencer/Private/DisplayNodes/SequencerDisplayNode.h
@@ -276,6 +276,7 @@ public:
 	// ICurveEditorTreeItem interface
 	virtual TSharedPtr<SWidget> GenerateCurveEditorTreeWidget(const FName& InColumnName, TWeakPtr<FCurveEditor> InCurveEditor, FCurveEditorTreeItemID InTreeItemID) override;
 	virtual void CreateCurveModels(TArray<TUniquePtr<FCurveModel>>& OutCurveModels) override;
+	virtual bool PassesFilter(const FCurveEditorTreeFilter* InFilter) const override;
 
 	/**
 	 * Get the display node that is ultimately responsible for constructing a section area widget for this node.
diff --git a/Engine/Source/Editor/Sequencer/Private/SSequencer.cpp b/Engine/Source/Editor/Sequencer/Private/SSequencer.cpp
--- a/Engine/Source/Editor/Sequencer/Private/SSequencer.cpp
+++ b/Engine/Source/Editor/Sequencer/Private/SSequencer.cpp
@@ -71,6 +71,7 @@
 #include "Camera/CameraActor.h"
 #include "SCurveEditorPanel.h"
 #include "Tree/SCurveEditorTree.h"
+#include "Tree/SCurveEditorTreeTextFilter.h"
 #include "SCurveKeyDetailPanel.h"
 #include "MovieSceneTimeHelpers.h"
 #include "FrameNumberNumericInterface.h"
@@ -410,7 +411,18 @@ void SSequencer::Construct(const FArguments& InArgs, TSharedRef<FSequencer> InSe
 			.DisabledTimeSnapTooltip(LOCTEXT("CurveEditorTimeSnapDisabledTooltip", "Time Snapping is currently driven by Sequencer."))
 			.TreeContent()
 			[
-				SNew(SCurveEditorTree, InSequencer->GetCurveEditor())
+				SNew(SVerticalBox)
+
+				+ SVerticalBox::Slot()
+				.AutoHeight()
+				[
+					SNew(SCurveEditorTreeTextFilter, InSequencer->GetCurveEditor())
+				]
+
+				+ SVerticalBox::Slot()
+				[
+					SNew(SCurveEditorTree, InSequencer->GetCurveEditor())
+				]
 			];
 
 			// Register an instanced custom property type layout to handle converting FFrameNumber from Tick Resolution to Display Rate.
diff --git a/Engine/Source/Editor/Sequencer/Private/SequencerNodeTree.cpp b/Engine/Source/Editor/Sequencer/Private/SequencerNodeTree.cpp
--- a/Engine/Source/Editor/Sequencer/Private/SequencerNodeTree.cpp
+++ b/Engine/Source/Editor/Sequencer/Private/SequencerNodeTree.cpp
@@ -703,6 +703,10 @@ TArray< TSharedRef<FSequencerDisplayNode> > FSequencerNodeTree::GetAllNodes() co
 void FSequencerNodeTree::UpdateCurveEditorTree()
 {
 	FCurveEditor* CurveEditor = Sequencer.GetCurveEditor().Get();
+
+	// Guard against multiple broadcasts here and defer them until the end of this function
+	FScopedCurveEditorTreeUpdateGuard ScopedUpdateGuard = CurveEditor->GetTree()->ScopedUpdateGuard();
+
 	auto Traverse_AddToCurveEditor = [this, CurveEditor](FSequencerDisplayNode& InNode)
 	{
 		if (InNode.GetType() == ESequencerNode::Track)
