Removed erroneous return

diff --git a/Engine/Source/Editor/CurveEditor/Public/Tree/CurveEditorTree.h b/Engine/Source/Editor/CurveEditor/Public/Tree/CurveEditorTree.h
--- a/Engine/Source/Editor/CurveEditor/Public/Tree/CurveEditorTree.h
+++ b/Engine/Source/Editor/CurveEditor/Public/Tree/CurveEditorTree.h
@@ -69,7 +69,7 @@ struct FCurveEditorOnChangedEvent
 	/**
 	 * Unbind a previously bound handler from this event
 	 */
-	void Unbind(FDelegateHandle Handle) { return Delegate.Remove(Handle); }
+	void Unbind(FDelegateHandle Handle) { Delegate.Remove(Handle); }
 
 private:
 
