Changes the snapping behavior in the Curve Editor to lock onto the first direction you move after pressing shift, instead of the bigger axis you've moved from. Releasing shift and then pressing again will allow you to change which direction you are locked onto.

diff --git a/Engine/Plugins/Editor/CurveEditorTools/Source/CurveEditorTools/Private/Tools/CurveEditorTransformTool.cpp b/Engine/Plugins/Editor/CurveEditorTools/Source/CurveEditorTools/Private/Tools/CurveEditorTransformTool.cpp
--- a/Engine/Plugins/Editor/CurveEditorTools/Source/CurveEditorTools/Private/Tools/CurveEditorTransformTool.cpp
+++ b/Engine/Plugins/Editor/CurveEditorTools/Source/CurveEditorTools/Private/Tools/CurveEditorTransformTool.cpp
@@ -435,6 +435,7 @@ void FCurveEditorTransformTool::OnDragStart()
 	}
 	TransformWidget.StartSize = TransformWidget.Size;
 	TransformWidget.StartPosition = TransformWidget.Position;
+	SnappingState.Reset();
 }
 
 void FCurveEditorTransformTool::OnDrag(const FPointerEvent& MouseEvent)
@@ -452,7 +453,7 @@ void FCurveEditorTransformTool::OnDrag(const FPointerEvent& MouseEvent)
 	// Dragging the center is the easy case!
 	if (TransformWidget.SelectedAnchorFlags == ECurveEditorAnchorFlags::Center)
 	{
-		FVector2D AxisLockedMousePosition = CurveEditor->GetAxisSnap().GetSnappedPosition(InitialMousePosition, MouseEvent.GetScreenSpacePosition(), MouseEvent);
+		FVector2D AxisLockedMousePosition = CurveEditor->GetAxisSnap().GetSnappedPosition(InitialMousePosition, MouseEvent.GetScreenSpacePosition(), MouseEvent, SnappingState);
 		{
 			FVector2D MouseDelta = AxisLockedMousePosition - InitialMousePosition;
 			TransformWidget.Position = TransformWidget.StartPosition + MouseDelta;
diff --git a/Engine/Plugins/Editor/CurveEditorTools/Source/CurveEditorTools/Private/Tools/CurveEditorTransformTool.h b/Engine/Plugins/Editor/CurveEditorTools/Source/CurveEditorTools/Private/Tools/CurveEditorTransformTool.h
--- a/Engine/Plugins/Editor/CurveEditorTools/Source/CurveEditorTools/Private/Tools/CurveEditorTransformTool.h
+++ b/Engine/Plugins/Editor/CurveEditorTools/Source/CurveEditorTools/Private/Tools/CurveEditorTransformTool.h
@@ -8,6 +8,7 @@
 #include "Curves/KeyHandle.h"
 #include "CurveDataAbstraction.h"
 #include "ScopedTransaction.h"
+#include "CurveEditorSnapMetrics.h"
 
 class FCurveEditor;
 struct FPointerEvent;
@@ -115,4 +116,6 @@ private:
 	TArray<FKeyData> KeysByCurve;
 
 	FVector2D InitialMousePosition;
+
+	FCurveEditorAxisSnap::FSnapState SnappingState;
 };
\ No newline at end of file
diff --git a/Engine/Source/Editor/CurveEditor/Private/DragOperations/CurveEditorDragOperation_MoveKeys.cpp b/Engine/Source/Editor/CurveEditor/Private/DragOperations/CurveEditorDragOperation_MoveKeys.cpp
--- a/Engine/Source/Editor/CurveEditor/Private/DragOperations/CurveEditorDragOperation_MoveKeys.cpp
+++ b/Engine/Source/Editor/CurveEditor/Private/DragOperations/CurveEditorDragOperation_MoveKeys.cpp
@@ -36,12 +36,14 @@ void FCurveEditorDragOperation_MoveKeys::OnBeginDrag(FVector2D InitialPosition,
 			Curve->GetKeyPositions(KeyData.Handles, KeyData.StartKeyPositions);
 		}
 	}
+
+	SnappingState.Reset();
 }
 
 void FCurveEditorDragOperation_MoveKeys::OnDrag(FVector2D InitialPosition, FVector2D CurrentPosition, const FPointerEvent& MouseEvent)
 {
 	TArray<FKeyPosition> NewKeyPositionScratch;
-	FVector2D MousePosition = CurveEditor->GetAxisSnap().GetSnappedPosition(InitialPosition, CurrentPosition, MouseEvent);
+	FVector2D MousePosition = CurveEditor->GetAxisSnap().GetSnappedPosition(InitialPosition, CurrentPosition, MouseEvent, SnappingState);
 
 	for (const FKeyData& KeyData : KeysByCurve)
 	{
diff --git a/Engine/Source/Editor/CurveEditor/Private/DragOperations/CurveEditorDragOperation_MoveKeys.h b/Engine/Source/Editor/CurveEditor/Private/DragOperations/CurveEditorDragOperation_MoveKeys.h
--- a/Engine/Source/Editor/CurveEditor/Private/DragOperations/CurveEditorDragOperation_MoveKeys.h
+++ b/Engine/Source/Editor/CurveEditor/Private/DragOperations/CurveEditorDragOperation_MoveKeys.h
@@ -43,4 +43,6 @@ private:
 
 	/** Key dragging data stored per-curve */
 	TArray<FKeyData> KeysByCurve;
+
+	FCurveEditorAxisSnap::FSnapState SnappingState;
 };
\ No newline at end of file
diff --git a/Engine/Source/Editor/CurveEditor/Private/DragOperations/CurveEditorDragOperation_Pan.cpp b/Engine/Source/Editor/CurveEditor/Private/DragOperations/CurveEditorDragOperation_Pan.cpp
--- a/Engine/Source/Editor/CurveEditor/Private/DragOperations/CurveEditorDragOperation_Pan.cpp
+++ b/Engine/Source/Editor/CurveEditor/Private/DragOperations/CurveEditorDragOperation_Pan.cpp
@@ -19,11 +19,12 @@ void FCurveEditorDragOperation_PanView::OnBeginDrag(FVector2D InitialPosition, F
 	InitialInputMax = ViewSpace.GetInputMax();
 	InitialOutputMin = ViewSpace.GetOutputMin();
 	InitialOutputMax = ViewSpace.GetOutputMax();
+	SnappingState.Reset();
 }
 
 void FCurveEditorDragOperation_PanView::OnDrag(FVector2D InitialPosition, FVector2D CurrentPosition, const FPointerEvent& MouseEvent)
 {
-	FVector2D PixelDelta = CurveEditor->GetAxisSnap().GetSnappedPosition(InitialPosition, CurrentPosition, MouseEvent, true) - InitialPosition;
+	FVector2D PixelDelta = CurveEditor->GetAxisSnap().GetSnappedPosition(InitialPosition, CurrentPosition, MouseEvent, SnappingState, true) - InitialPosition;
 
 	FCurveEditorScreenSpace ViewSpace = View->GetViewSpace();
 
@@ -46,11 +47,12 @@ void FCurveEditorDragOperation_PanInput::OnBeginDrag(FVector2D InitialPosition,
 	FCurveEditorScreenSpaceH InputSpace = CurveEditor->GetPanelInputSpace();
 	InitialInputMin = InputSpace.GetInputMin();
 	InitialInputMax = InputSpace.GetInputMax();
+	SnappingState.Reset();
 }
 
 void FCurveEditorDragOperation_PanInput::OnDrag(FVector2D InitialPosition, FVector2D CurrentPosition, const FPointerEvent& MouseEvent)
 {
-	FVector2D PixelDelta = CurveEditor->GetAxisSnap().GetSnappedPosition(InitialPosition, CurrentPosition, MouseEvent, true) - InitialPosition;
+	FVector2D PixelDelta = CurveEditor->GetAxisSnap().GetSnappedPosition(InitialPosition, CurrentPosition, MouseEvent, SnappingState, true) - InitialPosition;
 
 	FCurveEditorScreenSpaceH InputSpace = CurveEditor->GetPanelInputSpace();
 
diff --git a/Engine/Source/Editor/CurveEditor/Private/DragOperations/CurveEditorDragOperation_Pan.h b/Engine/Source/Editor/CurveEditor/Private/DragOperations/CurveEditorDragOperation_Pan.h
--- a/Engine/Source/Editor/CurveEditor/Private/DragOperations/CurveEditorDragOperation_Pan.h
+++ b/Engine/Source/Editor/CurveEditor/Private/DragOperations/CurveEditorDragOperation_Pan.h
@@ -22,6 +22,7 @@ private:
 
 	double InitialInputMin, InitialInputMax;
 	double InitialOutputMin, InitialOutputMax;
+	FCurveEditorAxisSnap::FSnapState SnappingState;
 };
 
 class FCurveEditorDragOperation_PanInput : public ICurveEditorDragOperation
@@ -36,4 +37,5 @@ private:
 
 	FCurveEditor* CurveEditor;
 	double InitialInputMin, InitialInputMax;
+	FCurveEditorAxisSnap::FSnapState SnappingState;
 };
\ No newline at end of file
diff --git a/Engine/Source/Editor/CurveEditor/Public/CurveEditorSnapMetrics.h b/Engine/Source/Editor/CurveEditor/Public/CurveEditorSnapMetrics.h
--- a/Engine/Source/Editor/CurveEditor/Public/CurveEditorSnapMetrics.h
+++ b/Engine/Source/Editor/CurveEditor/Public/CurveEditorSnapMetrics.h
@@ -5,6 +5,7 @@
 #include "Misc/FrameRate.h"
 #include "Math/Axis.h"
 
+
 /**
  * Utility struct that acts as a cache of the current snapping metrics for the curve editor
  */
@@ -47,6 +48,31 @@ struct FCurveEditorSnapMetrics
  */
 struct FCurveEditorAxisSnap
 {
+	/**
+	 * Snapping is not stateless but we want to manage it through the central area. This allows
+	 * state to be passed into from the calling area but still centralize the logic of handling it.
+	*/
+	struct FSnapState
+	{
+		FSnapState()
+		{
+			Reset();
+		}
+
+		void Reset()
+		{
+			MouseLockVector = FVector2D::UnitVector;
+			MousePosOnShiftStart = FVector2D::ZeroVector;
+			bHasPassedThreshold = false;
+			bHasStartPosition = false;
+		}
+
+		FVector2D MousePosOnShiftStart;
+		FVector2D MouseLockVector;
+		bool bHasPassedThreshold;
+		bool bHasStartPosition;
+	};
+
 	/** Can be set to either X, Y, or None to control which axis GetSnappedPosition snaps to. User can override None by pressing shift. */
 	EAxisList::Type RestrictedAxisList;
 
@@ -55,7 +81,7 @@ struct FCurveEditorAxisSnap
 		RestrictedAxisList = EAxisList::None;
 	}
 
-	FVector2D GetSnappedPosition(FVector2D InitialPosition, FVector2D CurrentPosition, const FPointerEvent& MouseEvent, const bool bIgnoreAxisLock = false)
+	FVector2D GetSnappedPosition(FVector2D InitialPosition, FVector2D CurrentPosition, const FPointerEvent& MouseEvent, FSnapState& InOutSnapState, const bool bIgnoreAxisLock = false)
 	{
 		check(RestrictedAxisList == EAxisList::Type::None || RestrictedAxisList == EAxisList::Type::X || RestrictedAxisList == EAxisList::Type::Y);
 
@@ -63,18 +89,47 @@ struct FCurveEditorAxisSnap
 		bool bCanUseShift = bIgnoreAxisLock || RestrictedAxisList == EAxisList::Type::None;
 
 		FVector2D MouseLockVector = FVector2D::UnitVector;
-		if (bCanUseShift && MouseEvent.IsShiftDown())
+		if (bCanUseShift)
 		{
-			// Figure out which axis is currently greater and snap to that 
-			FVector2D DragDelta = CurrentPosition - InitialPosition;
-			if (FMath::Abs(DragDelta.X) > FMath::Abs(DragDelta.Y))
+			// Snapping works by determining which direction the mouse is moving in while shift is being held
+			// down. Once you have moved far enough to pass the threshold, you're locked into that direction
+			// but if you release shift, it will reset the direction and allow you to pick a new direction.
+			if (MouseEvent.IsShiftDown())
 			{
-				MouseLockVector.Y = 0;
+				if (!InOutSnapState.bHasStartPosition)
+				{
+					InOutSnapState.MousePosOnShiftStart = CurrentPosition;
+					InOutSnapState.bHasStartPosition = true;
+				}
+				// If they have passed the threshold they should have a lock vector they're snapped to.
+				if (!InOutSnapState.bHasPassedThreshold)
+				{
+					InOutSnapState.MouseLockVector = MouseLockVector;
+					
+					// They have not passed the threshold yet, let's see if they've passed it now.
+					FVector2D DragDelta = CurrentPosition - InOutSnapState.MousePosOnShiftStart;
+					if (DragDelta.Size() > 2.f)
+					{
+						InOutSnapState.bHasPassedThreshold = true;
+						InOutSnapState.MouseLockVector = FVector2D::UnitVector;
+						if (FMath::Abs(DragDelta.X) > FMath::Abs(DragDelta.Y))
+						{
+							InOutSnapState.MouseLockVector.Y = 0;
+						}
+						else
+						{
+							InOutSnapState.MouseLockVector.X = 0;
+						}
+					}
+				}
 			}
 			else
 			{
-				MouseLockVector.X = 0;
+				// If they don't have shift pressed just disable the lock.
+				InOutSnapState.Reset();
 			}
+
+			MouseLockVector = InOutSnapState.MouseLockVector;
 		}
 		else if (!bIgnoreAxisLock && RestrictedAxisList == EAxisList::Type::X)
 		{
