Curve Editor: Fix dragging exit tangent (simple typo)

diff --git a/Engine/Source/Editor/CurveEditor/Private/DragOperations/CurveEditorDragOperation_Tangent.cpp b/Engine/Source/Editor/CurveEditor/Private/DragOperations/CurveEditorDragOperation_Tangent.cpp
--- a/Engine/Source/Editor/CurveEditor/Private/DragOperations/CurveEditorDragOperation_Tangent.cpp
+++ b/Engine/Source/Editor/CurveEditor/Private/DragOperations/CurveEditorDragOperation_Tangent.cpp
@@ -151,7 +151,7 @@ void FCurveEditorDragOperation_Tangent::OnDrag(FVector2D InitialPosition, FVecto
 					if (Attributes.HasTangentWeightMode() && Attributes.HasLeaveTangentWeight() &&
 						(Attributes.GetTangentWeightMode() == RCTWM_WeightedBoth || Attributes.GetTangentWeightMode() == RCTWM_WeightedLeave))
 					{
-						FVector2D TangentOffset = CurveEditor::ComputeScreenSpaceTangentOffset(CurveSpace, LeaveTangent, Attributes.GetArriveTangentWeight());
+						FVector2D TangentOffset = CurveEditor::ComputeScreenSpaceTangentOffset(CurveSpace, LeaveTangent, Attributes.GetLeaveTangentWeight());
 						TangentOffset += PixelDelta;
 
 						if (MouseEvent.IsShiftDown())
