Curve Editor: Tree filtering and synchronization improvements

This is a direct copy from CL 7095088 with a partial revert of original CL 6759985 to address the merge conflict

  - When changing selection in sequencer, the curve editor now selects the same nodes on its tree as well as isolating (auto-filtering) the tree to those nodes. This makes it much easier to drill down into curves from the sequence. Both these settings are configurable through the Editor Preferences ("Synchronize Curve Editor Selection" and "Isolate Curve Editor To Selection")
  - Changed pin widgets to only be visible when their table row is hovered, or if the pin is active
  - Changed filter logic to exclude implicit children when a sibling was explicitly matched. This prevents seemingly unmatched nodes showing up in hierarchies where nested matches occur.
  - Expanded text filter logic to include Parent.Child.Gradchild syntaxfor matching specific hierarchies of nodes

diff --git a/Engine/Source/Editor/CurveAssetEditor/Private/CurveAssetEditor.cpp b/Engine/Source/Editor/CurveAssetEditor/Private/CurveAssetEditor.cpp
--- a/Engine/Source/Editor/CurveAssetEditor/Private/CurveAssetEditor.cpp
+++ b/Engine/Source/Editor/CurveAssetEditor/Private/CurveAssetEditor.cpp
@@ -43,14 +43,14 @@ struct FCurveAssetEditorTreeItem : public ICurveEditorTreeItem
 		}
 	}
 
-	virtual TSharedPtr<SWidget> GenerateCurveEditorTreeWidget(const FName& InColumnName, TWeakPtr<FCurveEditor> InCurveEditor, FCurveEditorTreeItemID InTreeItemID) override
+	virtual TSharedPtr<SWidget> GenerateCurveEditorTreeWidget(const FName& InColumnName, TWeakPtr<FCurveEditor> InCurveEditor, FCurveEditorTreeItemID InTreeItemID, const TSharedRef<ITableRow>& TableRow) override
 	{
 		if (InColumnName == ColumnNames.Label)
 		{
 			return SNew(SHorizontalBox)
 
 			+ SHorizontalBox::Slot()
-			.Padding(FMargin(0.f, 0.f, 4.f, 0.f))
+			.Padding(FMargin(4.f))
 			.VAlign(VAlign_Center)
 			.HAlign(HAlign_Right)
 			.AutoWidth()
@@ -62,7 +62,7 @@ struct FCurveAssetEditorTreeItem : public ICurveEditorTreeItem
 		}
 		else if (InColumnName == ColumnNames.PinHeader)
 		{
-			return SNew(SCurveEditorTreePin, InCurveEditor, InTreeItemID);
+			return SNew(SCurveEditorTreePin, InCurveEditor, InTreeItemID, TableRow);
 		}
 
 		return nullptr;
diff --git a/Engine/Source/Editor/CurveEditor/Private/CurveEditor.cpp b/Engine/Source/Editor/CurveEditor/Private/CurveEditor.cpp
--- a/Engine/Source/Editor/CurveEditor/Private/CurveEditor.cpp
+++ b/Engine/Source/Editor/CurveEditor/Private/CurveEditor.cpp
@@ -213,31 +213,6 @@ ECurveEditorTreeSelectionState FCurveEditor::GetTreeSelectionState(FCurveEditorT
 	return Tree.GetSelectionState(InTreeItemID);
 }
 
-void FCurveEditor::SetDirectTreeSelection(TArray<FCurveEditorTreeItemID>&& TreeItems)
-{
-	TMap<FCurveEditorTreeItemID, ECurveEditorTreeSelectionState> PreviousSelection = Tree.GetSelection();
-	Tree.SetDirectSelection(MoveTemp(TreeItems));
-
-	const TMap<FCurveEditorTreeItemID, ECurveEditorTreeSelectionState>& NewSelection = Tree.GetSelection();
-	for (TTuple<FCurveEditorTreeItemID, ECurveEditorTreeSelectionState> CachedItem : PreviousSelection)
-	{
-		const ECurveEditorTreeSelectionState* NewState = NewSelection.Find(CachedItem.Key);
-		if (!NewState || *NewState == ECurveEditorTreeSelectionState::None)
-		{
-			GetTreeItem(CachedItem.Key).DestroyUnpinnedCurves(this);
-		}
-	}
-
-	// Ensure the new selection has valid curve models
-	for (TTuple<FCurveEditorTreeItemID, ECurveEditorTreeSelectionState> CachedItem : Tree.GetSelection())
-	{
-		if (CachedItem.Value != ECurveEditorTreeSelectionState::None)
-		{
-			GetTreeItem(CachedItem.Key).GetOrCreateCurves(this);
-		}
-	}
-}
-
 const TMap<FCurveEditorTreeItemID, ECurveEditorTreeSelectionState>& FCurveEditor::GetTreeSelection() const
 {
 	return Tree.GetSelection();
diff --git a/Engine/Source/Editor/CurveEditor/Private/Tree/CurveEditorTree.cpp b/Engine/Source/Editor/CurveEditor/Private/Tree/CurveEditorTree.cpp
--- a/Engine/Source/Editor/CurveEditor/Private/Tree/CurveEditorTree.cpp
+++ b/Engine/Source/Editor/CurveEditor/Private/Tree/CurveEditorTree.cpp
@@ -5,7 +5,9 @@
 #include "Tree/CurveEditorTreeFilter.h"
 #include "CurveEditor.h"
 
+#include "Containers/SortedMap.h"
 #include "Algo/AnyOf.h"
+#include "Algo/AllOf.h"
 
 class FCurveModel;
 
@@ -52,35 +54,66 @@ void FCurveEditorTreeItem::DestroyUnpinnedCurves(FCurveEditor* CurveEditor)
 	}
 }
 
-FScopedCurveEditorTreeUpdateGuard::FScopedCurveEditorTreeUpdateGuard(FCurveEditorTree* InTree)
+void FCurveEditorTreeDelegate::Broadcast()
+{
+	if (ensureAlwaysMsgf(!bBroadcasting, TEXT("Attempting to broadcast changes while already broadcasting changes. This may be an infinite loop and is not allowed. Do not attempt to mutate the tree in response to it being mutated")))
+	{
+		TGuardValue<bool> ScopeGuard(bBroadcasting, true);
+		FSimpleMulticastDelegate::Broadcast();
+	}
+}
+
+FScopedCurveEditorTreeEventGuard::FScopedCurveEditorTreeEventGuard(FCurveEditorTree* InTree)
 	: Tree(InTree)
 {
-	Tree->OnChanged()->UpdateGuardCounter += 1;
+	CachedItemSerialNumber = Tree->Events.OnItemsChanged.SerialNumber;
+	CachedSelectionSerialNumber = Tree->Events.OnSelectionChanged.SerialNumber;
+	CachedFiltersSerialNumber = Tree->Events.OnFiltersChanged.SerialNumber;
+	Tree->Events.UpdateGuardCounter += 1;
 }
 
-FScopedCurveEditorTreeUpdateGuard::FScopedCurveEditorTreeUpdateGuard(FScopedCurveEditorTreeUpdateGuard&& RHS)
+FScopedCurveEditorTreeEventGuard::FScopedCurveEditorTreeEventGuard(FScopedCurveEditorTreeEventGuard&& RHS)
 	: Tree(RHS.Tree)
 {
-	Tree->OnChanged()->UpdateGuardCounter += 1;
+	Tree->Events.UpdateGuardCounter += 1;
 }
 
-FScopedCurveEditorTreeUpdateGuard& FScopedCurveEditorTreeUpdateGuard::operator=(FScopedCurveEditorTreeUpdateGuard&& RHS)
+FScopedCurveEditorTreeEventGuard& FScopedCurveEditorTreeEventGuard::operator=(FScopedCurveEditorTreeEventGuard&& RHS)
 {
 	if (&RHS != this)
 	{
 		Tree = RHS.Tree;
-		Tree->OnChanged()->UpdateGuardCounter += 1;
+		Tree->Events.UpdateGuardCounter += 1;
 	}
 	return *this;
 }
 
-FScopedCurveEditorTreeUpdateGuard::~FScopedCurveEditorTreeUpdateGuard()
+FScopedCurveEditorTreeEventGuard::~FScopedCurveEditorTreeEventGuard()
 {
-	FCurveEditorOnChangedEvent* OnChangedEvent = Tree->OnChanged();
-	if (--OnChangedEvent->UpdateGuardCounter == 0)
+	// Before we decrement the UpdateGuardCounter, we check to see whether a re-filter is required so that
+	// the re-filter operation itself remains guarded from causing subsequent updates
+	const bool bRequiresFilter = (Tree->Events.OnItemsChanged.SerialNumber != CachedItemSerialNumber) || (Tree->Events.OnFiltersChanged.SerialNumber != CachedFiltersSerialNumber);
+	if (bRequiresFilter && Tree->Events.UpdateGuardCounter == 1)
 	{
-		TGuardValue<bool> Guard(OnChangedEvent->bIsBroadcastInProgress, true);
-		OnChangedEvent->Delegate.Broadcast();
+		Tree->RunFilters();
+	}
+
+	if (--Tree->Events.UpdateGuardCounter == 0)
+	{
+		if (Tree->Events.OnItemsChanged.SerialNumber != CachedItemSerialNumber)
+		{
+			Tree->Events.OnItemsChanged.Broadcast();
+		}
+
+		if (Tree->Events.OnSelectionChanged.SerialNumber != CachedSelectionSerialNumber)
+		{
+			Tree->Events.OnSelectionChanged.Broadcast();
+		}
+
+		if (Tree->Events.OnFiltersChanged.SerialNumber != CachedFiltersSerialNumber)
+		{
+			Tree->Events.OnFiltersChanged.Broadcast();
+		}
 	}
 }
 
@@ -116,8 +149,7 @@ const TMap<FCurveEditorTreeItemID, FCurveEditorTreeItem>& FCurveEditorTree::GetA
 
 FCurveEditorTreeItem* FCurveEditorTree::AddItem(FCurveEditorTreeItemID ParentID)
 {
-	checkf(!OnChangedEvent.IsBroadcastInProgress(), TEXT("Curve editor tree must not be manipulated in response to it changing"));
-	FScopedCurveEditorTreeUpdateGuard BroadcastChangeUpdate(this);
+	FScopedCurveEditorTreeEventGuard EventGuard(this);
 
 	FCurveEditorTreeItemID NewItemID = NextTreeItemID;
 
@@ -130,14 +162,14 @@ FCurveEditorTreeItem* FCurveEditorTree::AddItem(FCurveEditorTreeItemID ParentID)
 	ParentContainer.bRequiresSort = true;
 
 	++NextTreeItemID.Value;
+	++Events.OnItemsChanged.SerialNumber;
 
 	return &NewItem;
 }
 
 void FCurveEditorTree::RemoveItem(FCurveEditorTreeItemID ItemID, FCurveEditor* CurveEditor)
 {
-	checkf(!OnChangedEvent.IsBroadcastInProgress(), TEXT("Curve editor tree must not be manipulated in response to it changing"));
-	FScopedCurveEditorTreeUpdateGuard BroadcastChangeUpdate(this);
+	FScopedCurveEditorTreeEventGuard EventGuard(this);
 
 	FCurveEditorTreeItem* Item = Items.Find(ItemID);
 	if (!Item)
@@ -156,8 +188,14 @@ void FCurveEditorTree::RemoveItem(FCurveEditorTreeItemID ItemID, FCurveEditor* C
 
 	// Item is no longer valid
 
-	Items.Remove(ItemID);
-	Selection.Remove(ItemID);
+	if (Items.Remove(ItemID) != 0)
+	{
+		++Events.OnItemsChanged.SerialNumber;
+	}
+	if (Selection.Remove(ItemID) != 0)
+	{
+		++Events.OnSelectionChanged.SerialNumber;
+	}
 }
 
 void FCurveEditorTree::RemoveChildrenRecursive(TArray<FCurveEditorTreeItemID>&& LocalChildren, FCurveEditor* CurveEditor)
@@ -171,26 +209,53 @@ void FCurveEditorTree::RemoveChildrenRecursive(TArray<FCurveEditorTreeItemID>&&
 
 			RemoveChildrenRecursive(MoveTemp(ChildItem->Children.ChildIDs), CurveEditor);
 
-			Items.Remove(ChildID);
-			Selection.Remove(ChildID);
+			if (Items.Remove(ChildID) != 0)
+			{
+				++Events.OnItemsChanged.SerialNumber;
+			}
+			if (Selection.Remove(ChildID) != 0)
+			{
+				++Events.OnSelectionChanged.SerialNumber;
+			}
 		}
 	}
 }
 
-
-bool FCurveEditorTree::FilterSpecificItems(TArrayView<FCurveEditorTreeFilter* const> FilterPtrs, TArrayView<const FCurveEditorTreeItemID> ItemsToFilter, ECurveEditorTreeFilterState InheritedState)
+bool FCurveEditorTree::PerformFilterPass(TArrayView<const FCurveEditorTreeFilter* const> FilterPtrs, TArrayView<const FCurveEditorTreeItemID> ItemsToFilter, ECurveEditorTreeFilterState InheritedState)
 {
 	bool bAnyMatched = false;
 
 	for (FCurveEditorTreeItemID ItemID : ItemsToFilter)
 	{
+		// If it failed a previous pass, don't consider it for this pass
+		if (FilterStates.Get(ItemID) != ECurveEditorTreeFilterState::NoMatch)
+		{
+			TSharedPtr<ICurveEditorTreeItem> TreeItemImpl = GetItem(ItemID).GetItem();
+			if (TreeItemImpl)
+			{
+				const bool bMatchesFilter = Algo::AnyOf(FilterPtrs, [TreeItemImpl](const FCurveEditorTreeFilter* Filter){ return TreeItemImpl->PassesFilter(Filter); });
+				if (bMatchesFilter)
+				{
+					InheritedState = ECurveEditorTreeFilterState::NoMatch;
+					break;
+				}
+			}
+		}
+	}
+
+	for (FCurveEditorTreeItemID ItemID : ItemsToFilter)
+	{
+		// If it failed a previous pass, don't consider it for this pass
+		if (FilterStates.Get(ItemID) == ECurveEditorTreeFilterState::NoMatch)
+		{
+			continue;
+		}
+
 		const FCurveEditorTreeItem& TreeItem = GetItem(ItemID);
 
-		// Retrieve the existing filter state for this item. This may have been set by preceeding filters.
 		ECurveEditorTreeFilterState FilterState         = InheritedState;
 		ECurveEditorTreeFilterState ChildInheritedState = InheritedState;
 
-		// If this has already not been matched, run it through our filter
 		TSharedPtr<ICurveEditorTreeItem> TreeItemImpl = TreeItem.GetItem();
 		if (TreeItemImpl)
 		{
@@ -204,7 +269,7 @@ bool FCurveEditorTree::FilterSpecificItems(TArrayView<FCurveEditorTreeFilter* co
 		}
 
 		// Run the filter on all child nodes
-		const bool bMatchedChildren = FilterSpecificItems(FilterPtrs, TreeItem.GetChildren(), ChildInheritedState);
+		const bool bMatchedChildren = PerformFilterPass(FilterPtrs, TreeItem.GetChildren(), ChildInheritedState);
 
 		// If we matched children we become an implicit parent if not already matched
 		if (bMatchedChildren && FilterState != ECurveEditorTreeFilterState::Match)
@@ -213,10 +278,7 @@ bool FCurveEditorTree::FilterSpecificItems(TArrayView<FCurveEditorTreeFilter* co
 			FilterState = ECurveEditorTreeFilterState::ImplicitParent;
 		}
 
-		if (FilterState != ECurveEditorTreeFilterState::NoMatch)
-		{
-			FilterStates.SetFilterState(ItemID, FilterState);
-		}
+		FilterStates.SetFilterState(ItemID, FilterState);
 	}
 
 	return bAnyMatched;
@@ -224,49 +286,95 @@ bool FCurveEditorTree::FilterSpecificItems(TArrayView<FCurveEditorTreeFilter* co
 
 void FCurveEditorTree::RunFilters()
 {
-	checkf(!OnChangedEvent.IsBroadcastInProgress(), TEXT("Curve editor tree must not be manipulated in response to it changing"));
+	FScopedCurveEditorTreeEventGuard EventGuard(this);
 
-	FScopedCurveEditorTreeUpdateGuard BroadcastChangeUpdate(this);
-
-	// Reset all the filter states back to the default
+	// Always reset filter states
 	FilterStates.Reset();
 
-	if (Filters.Num())
+	if (WeakFilters.Num() != 0)
 	{
-		FilterStates.Activate();
-
-		TArray<FCurveEditorTreeFilter*> FilterPtrs;
+		// Gather all valid filters into a map sorted by pass
+		TSortedMap<int32, TArray<const FCurveEditorTreeFilter*>> FiltersByPass;
 
-		for (int32 Index = Filters.Num() - 1; Index >= 0; --Index)
+		for (int32 Index = WeakFilters.Num() - 1; Index >= 0; --Index)
 		{
-			TSharedPtr<FCurveEditorTreeFilter> Filter = Filters[Index].Pin();
+			TSharedPtr<FCurveEditorTreeFilter> Filter = WeakFilters[Index].Pin();
 			if (!Filter)
 			{
 				// Remove invalid filters
-				Filters.RemoveAtSwap(Index, 1, false);
+				WeakFilters.RemoveAtSwap(Index, 1, false);
 			}
 			else
 			{
-				FilterPtrs.Add(Filter.Get());
+				FiltersByPass.FindOrAdd(Filter->GetFilterPass()).Add(Filter.Get());
 			}
 		}
 
-		FilterSpecificItems(FilterPtrs, RootItems.ChildIDs, ECurveEditorTreeFilterState::NoMatch);
+		// Deactivate the filter states for the first pass to ensure that _all_ items are considered for filtering in the very first pass
+		FilterStates.Deactivate();
+
+		for (const TTuple<int32, TArray<const FCurveEditorTreeFilter*>>& Pair : FiltersByPass)
+		{
+			PerformFilterPass(Pair.Value, RootItems.ChildIDs, ECurveEditorTreeFilterState::NoMatch);
+
+			// Activate the filters after the first pass so that subsequent passes only consider items that have previously matched in some way.
+			// This ensures that each pass works as a boolean AND.
+			FilterStates.Activate();
+		}
 	}
 	else
 	{
 		FilterStates.Deactivate();
 	}
+
+	++Events.OnItemsChanged.SerialNumber;
+	++Events.OnFiltersChanged.SerialNumber;
 }
 
 void FCurveEditorTree::AddFilter(TWeakPtr<FCurveEditorTreeFilter> NewFilter)
 {
-	Filters.AddUnique(NewFilter);
+	FScopedCurveEditorTreeEventGuard EventGuard(this);
+
+	WeakFilters.AddUnique(NewFilter);
+
+	// Always broadcast the filter change even if the filter already existed since calling AddFilter indicates a clear intent to have the tree re-filtered
+	++Events.OnFiltersChanged.SerialNumber;
 }
 
 void FCurveEditorTree::RemoveFilter(TWeakPtr<FCurveEditorTreeFilter> FilterToRemove)
 {
-	Filters.Remove(FilterToRemove);
+	FScopedCurveEditorTreeEventGuard EventGuard(this);
+
+	WeakFilters.Remove(FilterToRemove);
+	// Always broadcast the filter change even if the filter already existed since calling RemoveFilter indicates a clear intent to have the tree re-filtered
+	++Events.OnFiltersChanged.SerialNumber;
+}
+
+const FCurveEditorTreeFilter* FCurveEditorTree::FindFilterByType(ECurveEditorTreeFilterType Type) const
+{
+	for (const TWeakPtr<FCurveEditorTreeFilter>& WeakFilter : WeakFilters)
+	{
+		TSharedPtr<FCurveEditorTreeFilter> Pinned = WeakFilter.Pin();
+		if (Pinned && Pinned->GetType() == Type)
+		{
+			return Pinned.Get();
+		}
+	}
+
+	return nullptr;
+}
+
+TArrayView<const TWeakPtr<FCurveEditorTreeFilter>> FCurveEditorTree::GetFilters() const
+{
+	return WeakFilters;
+}
+
+void FCurveEditorTree::ClearFilters()
+{
+	FScopedCurveEditorTreeEventGuard EventGuard(this);
+
+	WeakFilters.Empty();
+	++Events.OnFiltersChanged.SerialNumber;
 }
 
 const FCurveEditorFilterStates& FCurveEditorTree::GetFilterStates() const
@@ -279,8 +387,11 @@ ECurveEditorTreeFilterState FCurveEditorTree::GetFilterState(FCurveEditorTreeIte
 	return FilterStates.Get(InTreeItemID);
 }
 
-void FCurveEditorTree::SetDirectSelection(TArray<FCurveEditorTreeItemID>&& TreeItems)
+void FCurveEditorTree::SetDirectSelection(TArray<FCurveEditorTreeItemID>&& TreeItems, FCurveEditor* InCurveEditor)
 {
+	FScopedCurveEditorTreeEventGuard EventGuard(this);
+
+	TMap<FCurveEditorTreeItemID, ECurveEditorTreeSelectionState> PreviousSelection = MoveTemp(Selection);
 	Selection.Reset();
 
 	// Recursively add child items
@@ -300,9 +411,35 @@ void FCurveEditorTree::SetDirectSelection(TArray<FCurveEditorTreeItemID>&& TreeI
 
 		for (FCurveEditorTreeItemID ChildID : Items.FindChecked(ItemID).GetChildren())
 		{
-			TreeItems.Add(ChildID);
+			if (FilterStates.Get(ChildID) != ECurveEditorTreeFilterState::NoMatch)
+			{
+				TreeItems.Add(ChildID);
+			}
+		}
+	}
+
+	for (TTuple<FCurveEditorTreeItemID, ECurveEditorTreeSelectionState> OldItem : PreviousSelection)
+	{
+		const ECurveEditorTreeSelectionState* NewState = Selection.Find(OldItem.Key);
+		if (!NewState || *NewState == ECurveEditorTreeSelectionState::None)
+		{
+			GetItem(OldItem.Key).DestroyUnpinnedCurves(InCurveEditor);
+		}
+	}
+
+	// Ensure the new selection has valid curve models
+	for (TTuple<FCurveEditorTreeItemID, ECurveEditorTreeSelectionState> NewItem : Selection)
+	{
+		if (NewItem.Value != ECurveEditorTreeSelectionState::None)
+		{
+			GetItem(NewItem.Key).GetOrCreateCurves(InCurveEditor);
 		}
 	}
+
+	if (!PreviousSelection.OrderIndependentCompareEqual(Selection))
+	{
+		++Events.OnSelectionChanged.SerialNumber;
+	}
 }
 
 const TMap<FCurveEditorTreeItemID, ECurveEditorTreeSelectionState>& FCurveEditorTree::GetSelection() const
diff --git a/Engine/Source/Editor/CurveEditor/Private/Tree/CurveEditorTreeFilter.cpp b/Engine/Source/Editor/CurveEditor/Private/Tree/CurveEditorTreeFilter.cpp
--- a/Engine/Source/Editor/CurveEditor/Private/Tree/CurveEditorTreeFilter.cpp
+++ b/Engine/Source/Editor/CurveEditor/Private/Tree/CurveEditorTreeFilter.cpp
@@ -21,4 +21,29 @@ ECurveEditorTreeFilterType FCurveEditorTreeFilter::RegisterFilterType()
 	NextFilterType = ECurveEditorTreeFilterType( ((__underlying_type(ECurveEditorTreeFilterType))NextFilterType) + 1 );
 
 	return NextFilterType;
+}
+
+void FCurveEditorTreeTextFilter::AssignFromText(const FString& FilterString)
+{
+	ChildToParentFilterTerms.Reset();
+
+	static const bool bCullEmpty = true;
+
+	TArray<FString> FilterTerms;
+	FilterString.ParseIntoArray(FilterTerms, TEXT(" "), bCullEmpty);
+
+	TArray<FString> ParentToChildTerms;
+	for (const FString& Term : FilterTerms)
+	{
+		ParentToChildTerms.Reset();
+		Term.ParseIntoArray(ParentToChildTerms, TEXT("."), bCullEmpty);
+
+		// Move the results into a new term in reverse order (so they are child -> parent)
+		FCurveEditorTreeTextFilterTerm NewTerm;
+		for (int32 Index = ParentToChildTerms.Num()-1; Index >= 0; --Index)
+		{
+			NewTerm.ChildToParentTokens.Emplace(FCurveEditorTreeTextFilterToken{ MoveTemp(ParentToChildTerms[Index]) });
+		}
+		ChildToParentFilterTerms.Emplace(MoveTemp(NewTerm));
+	}
 }
\ No newline at end of file
diff --git a/Engine/Source/Editor/CurveEditor/Private/Tree/SCurveEditorTree.cpp b/Engine/Source/Editor/CurveEditor/Private/Tree/SCurveEditorTree.cpp
--- a/Engine/Source/Editor/CurveEditor/Private/Tree/SCurveEditorTree.cpp
+++ b/Engine/Source/Editor/CurveEditor/Private/Tree/SCurveEditorTree.cpp
@@ -7,63 +7,8 @@
 #include "CurveEditor.h"
 
 #include "Widgets/Views/STableRow.h"
-#include "Widgets/Views/STreeView.h"
 #include "Widgets/Views/SExpanderArrow.h"
 
-template <> struct TListTypeTraits<FCurveEditorTreeItemID>
-{
-public:
-	struct NullableType : FCurveEditorTreeItemID
-	{
-		NullableType(TYPE_OF_NULLPTR){}
-		NullableType(FCurveEditorTreeItemID Other) : FCurveEditorTreeItemID(Other) {}
-	};
-
-	using MapKeyFuncs = TDefaultMapHashableKeyFuncs<FCurveEditorTreeItemID, TSharedRef<ITableRow>, false>;
-	using MapKeyFuncsSparse = TDefaultMapHashableKeyFuncs<FCurveEditorTreeItemID, FSparseItemInfo, false>;
-	using SetKeyFuncs = DefaultKeyFuncs<FCurveEditorTreeItemID>;
-
-	template<typename U>
-	static void AddReferencedObjects(FReferenceCollector&, TArray<FCurveEditorTreeItemID>&, TSet<FCurveEditorTreeItemID>&, TMap< const U*, FCurveEditorTreeItemID >&) {}
-
-	static bool IsPtrValid(NullableType InPtr)
-	{
-		return InPtr.IsValid();
-	}
-
-	static void ResetPtr(NullableType& InPtr)
-	{
-		InPtr = nullptr;
-	}
-
-	static NullableType MakeNullPtr()
-	{
-		return nullptr;
-	}
-
-	static FCurveEditorTreeItemID NullableItemTypeConvertToItemType(NullableType InPtr)
-	{
-		return InPtr;
-	}
-
-	static FString DebugDump(FCurveEditorTreeItemID InPtr)
-	{
-		return FString::Printf(TEXT("%d"), InPtr.GetValue());
-	}
-
-	class SerializerType{};
-};
-template <>
-struct TIsValidListItem<FCurveEditorTreeItemID>
-{
-	enum
-	{
-		Value = true
-	};
-};
-
-class SCurveEditorTreeView : public STreeView<FCurveEditorTreeItemID>
-{};
 
 struct SCurveEditorTableRow : SMultiColumnTableRow<FCurveEditorTreeItemID>
 {
@@ -93,10 +38,12 @@ struct SCurveEditorTableRow : SMultiColumnTableRow<FCurveEditorTreeItemID>
 		TSharedPtr<FCurveEditor> CurveEditor = WeakCurveEditor.Pin();
 		TSharedPtr<ICurveEditorTreeItem> TreeItem = CurveEditor ? CurveEditor->GetTreeItem(TreeItemID).GetItem() : nullptr;
 
+		TSharedRef<ITableRow> TableRow = SharedThis(this);
+
 		TSharedPtr<SWidget> Widget;
 		if (TreeItem.IsValid())
 		{
-			Widget = TreeItem->GenerateCurveEditorTreeWidget(InColumnName, WeakCurveEditor, TreeItemID);
+			Widget = TreeItem->GenerateCurveEditorTreeWidget(InColumnName, WeakCurveEditor, TreeItemID, TableRow);
 		}
 
 		if (!Widget.IsValid())
@@ -125,6 +72,7 @@ struct SCurveEditorTableRow : SMultiColumnTableRow<FCurveEditorTreeItemID>
 	}
 };
 
+
 void SCurveEditorTree::Construct(const FArguments& InArgs, TSharedPtr<FCurveEditor> InCurveEditor)
 {
 	bFilterWasActive = false;
@@ -138,9 +86,8 @@ void SCurveEditorTree::Construct(const FArguments& InArgs, TSharedPtr<FCurveEdit
 
 	CurveEditor = InCurveEditor;
 
-	ChildSlot
-	[
-		SAssignNew(TreeView, SCurveEditorTreeView)
+	STreeView<FCurveEditorTreeItemID>::Construct(
+		STreeView<FCurveEditorTreeItemID>::FArguments()
 		.SelectionMode(ESelectionMode::Multi)
 		.HeaderRow(HeaderRow)
 		.HighlightParentNodesForSelection(true)
@@ -154,9 +101,9 @@ void SCurveEditorTree::Construct(const FArguments& InArgs, TSharedPtr<FCurveEdit
 				this->OnTreeSelectionChanged(InItemID, Type);
 			}
 		)
-	];
+	);
 
-	CurveEditor->GetTree()->Bind_OnChanged(FSimpleDelegate::CreateSP(this, &SCurveEditorTree::RefreshTree));
+	CurveEditor->GetTree()->Events.OnItemsChanged.AddSP(this, &SCurveEditorTree::RefreshTree);
 }
 
 void SCurveEditorTree::RefreshTree()
@@ -171,15 +118,27 @@ void SCurveEditorTree::RefreshTree()
 	{
 		// Save expansion states
 		PreFilterExpandedItems.Reset();
-		TreeView->GetExpandedItems(PreFilterExpandedItems);
+		GetExpandedItems(PreFilterExpandedItems);
 	}
 	else if (!FilterStates.IsActive() && bFilterWasActive)
 	{
+		// Add any currently selected items' parents to the expanded items array.
+		// This ensures that items that were selected during a filter operation remain expanded and selected when finished
+		for (FCurveEditorTreeItemID SelectedItem : GetSelectedItems())
+		{
+			FCurveEditorTreeItemID ItemToExpand = CurveEditorTree->GetItem(SelectedItem).GetParentID();
+			while (ItemToExpand != FCurveEditorTreeItemID::Invalid())
+			{
+				PreFilterExpandedItems.Add(ItemToExpand);
+				ItemToExpand = CurveEditorTree->GetItem(ItemToExpand).GetParentID();
+			}
+		}
+
 		// Restore expansion states
-		TreeView->ClearExpandedItems();
+		ClearExpandedItems();
 		for (FCurveEditorTreeItemID ExpandedItem : PreFilterExpandedItems)
 		{
-			TreeView->SetItemExpansion(ExpandedItem, true);
+			SetItemExpansion(ExpandedItem, true);
 		}
 		PreFilterExpandedItems.Reset();
 	}
@@ -194,12 +153,12 @@ void SCurveEditorTree::RefreshTree()
 	}
 
 	RootItems.Shrink();
-	TreeView->RequestTreeRefresh();
+	RequestTreeRefresh();
 
 	if (FilterStates.IsActive())
 	{
 		// If a filter is active, all matched items and their parents are expanded
-		TreeView->ClearExpandedItems();
+		ClearExpandedItems();
 		for (const TTuple<FCurveEditorTreeItemID, FCurveEditorTreeItem>& Pair : CurveEditorTree->GetAllItems())
 		{
 			ECurveEditorTreeFilterState FilterState = FilterStates.Get(Pair.Key);
@@ -207,7 +166,7 @@ void SCurveEditorTree::RefreshTree()
 			// Expand any matched items or parents of matched items
 			if (FilterState == ECurveEditorTreeFilterState::Match || FilterState == ECurveEditorTreeFilterState::ImplicitParent)
 			{
-				TreeView->SetItemExpansion(Pair.Key, true);
+				SetItemExpansion(Pair.Key, true);
 			}
 		}
 	}
@@ -219,7 +178,7 @@ FReply SCurveEditorTree::OnKeyDown(const FGeometry& MyGeometry, const FKeyEvent&
 {
 	if (InKeyEvent.GetKey() == EKeys::Escape)
 	{
-		TreeView->ClearSelection();
+		ClearSelection();
 		return FReply::Handled();
 	}
 
@@ -246,14 +205,14 @@ void SCurveEditorTree::GetTreeItemChildren(FCurveEditorTreeItemID Parent, TArray
 
 void SCurveEditorTree::OnTreeSelectionChanged(FCurveEditorTreeItemID, ESelectInfo::Type)
 {
-	CurveEditor->SetDirectTreeSelection(TreeView->GetSelectedItems());
+	CurveEditor->GetTree()->SetDirectSelection(GetSelectedItems(), CurveEditor.Get());
 }
 
 void SCurveEditorTree::SetItemExpansionRecursive(FCurveEditorTreeItemID Model, bool bInExpansionState)
 {
 	if (Model.IsValid())
 	{
-		TreeView->SetItemExpansion(Model, bInExpansionState);
+		SetItemExpansion(Model, bInExpansionState);
 
 		TArray<FCurveEditorTreeItemID> Children;
 		GetTreeItemChildren(Model, Children);
diff --git a/Engine/Source/Editor/CurveEditor/Private/Tree/SCurveEditorTreeFilterStatusBar.cpp b/Engine/Source/Editor/CurveEditor/Private/Tree/SCurveEditorTreeFilterStatusBar.cpp
new file mode 100644
--- /dev/null
+++ b/Engine/Source/Editor/CurveEditor/Private/Tree/SCurveEditorTreeFilterStatusBar.cpp
@@ -0,0 +1,125 @@
+// Copyright 1998-2019 Epic Games, Inc. All Rights Reserved.
+
+#include "Tree/SCurveEditorTreeFilterStatusBar.h"
+#include "CurveEditor.h"
+
+#include "Algo/Count.h"
+#include "Widgets/SBoxPanel.h"
+#include "Widgets/Text/STextBlock.h"
+#include "Widgets/Input/SHyperlink.h"
+
+#define LOCTEXT_NAMESPACE "SCurveEditorTreeFilterStatusBar"
+
+void SCurveEditorTreeFilterStatusBar::Construct(const FArguments& InArgs, TSharedPtr<FCurveEditor> InCurveEditor)
+{
+	WeakCurveEditor = InCurveEditor;
+
+	ChildSlot
+	.Padding(FMargin(5.f, 3.f))
+	[
+		SNew(SHorizontalBox)
+
+		+ SHorizontalBox::Slot()
+		[
+			SAssignNew(TextBlock, STextBlock)
+		]
+
+		+ SHorizontalBox::Slot()
+		.AutoWidth()
+		.Padding(FMargin(3.f, 0.f, 0.f, 0.f))
+		[
+			SNew(SHyperlink)
+			.Visibility(this, &SCurveEditorTreeFilterStatusBar::GetVisibilityFromFilter)
+			.Text(LOCTEXT("ClearFilters", "clear"))
+			.OnNavigate(this, &SCurveEditorTreeFilterStatusBar::ClearFilters)
+		]
+	];
+
+	InCurveEditor->GetTree()->Events.OnItemsChanged.AddSP(this, &SCurveEditorTreeFilterStatusBar::UpdateText);
+	InCurveEditor->GetTree()->Events.OnSelectionChanged.AddSP(this, &SCurveEditorTreeFilterStatusBar::UpdateText);
+}
+
+void SCurveEditorTreeFilterStatusBar::ClearFilters()
+{
+	TSharedPtr<FCurveEditor> CurveEditor = WeakCurveEditor.Pin();
+	if (CurveEditor)
+	{
+		CurveEditor->GetTree()->ClearFilters();
+	}
+}
+
+EVisibility SCurveEditorTreeFilterStatusBar::GetVisibilityFromFilter() const
+{
+	TSharedPtr<FCurveEditor> CurveEditor = WeakCurveEditor.Pin();
+	return (CurveEditor && CurveEditor->GetTree()->GetFilterStates().IsActive()) ? EVisibility::Visible : EVisibility::Collapsed;
+}
+
+void SCurveEditorTreeFilterStatusBar::UpdateText()
+{
+	TSharedPtr<FCurveEditor> CurveEditor = WeakCurveEditor.Pin();
+	if (!ensureAlways(CurveEditor))
+	{
+		return;
+	}
+
+	FText NewText;
+	FLinearColor NewColor = FLinearColor::White;
+
+	const FCurveEditorTree* Tree = CurveEditor->GetTree();
+	const FCurveEditorFilterStates& FilterStates = Tree->GetFilterStates();
+	const TMap<FCurveEditorTreeItemID, ECurveEditorTreeSelectionState>& Selection = Tree->GetSelection();
+
+	FFormatNamedArguments NamedArgs;
+	NamedArgs.Add("Total", Tree->GetAllItems().Num());
+
+	const bool bHasSelection = Selection.Num() != 0;
+	const bool bHasFilter    = FilterStates.IsActive();
+
+	if (bHasSelection)
+	{
+		const int32 NumExplicitlySelected = Algo::CountIf(Selection, [](TTuple<FCurveEditorTreeItemID, ECurveEditorTreeSelectionState> Pair){ return Pair.Value == ECurveEditorTreeSelectionState::Explicit; });
+		NamedArgs.Add("NumSelected", NumExplicitlySelected);
+	}
+
+	if (bHasFilter)
+	{
+		NamedArgs.Add("NumMatched", FilterStates.GetNumMatched() + FilterStates.GetNumMatchedImplicitly());
+	}
+
+
+	if (bHasFilter)
+	{
+		if (FilterStates.GetNumMatched() == 0)
+		{
+			// Red = no matched
+			NewColor = FLinearColor( 1.0f, 0.4f, 0.4f );
+		}
+		else
+		{
+			// Green = matched filter
+			NewColor = FLinearColor( 0.4f, 1.0f, 0.4f );
+		}
+
+		if (bHasSelection)
+		{
+			NewText = FText::Format(LOCTEXT("FilteredStatus_WithSelection", "Showing {NumMatched} of {Total} items ({NumSelected} selected)"), NamedArgs);
+		}
+		else 
+		{
+			NewText = FText::Format(LOCTEXT("FilteredStatus_NoSelection", "Showing {NumMatched} of {Total} items"), NamedArgs);
+		}
+	}
+	else if (bHasSelection)
+	{
+		NewText = FText::Format(LOCTEXT("UnfilteredStatus_WithSelection", "{Total} items ({NumSelected} selected)"), NamedArgs);
+	}
+	else
+	{
+		NewText = FText::Format(LOCTEXT("UnfilteredStatus_NoSelection", "{Total} items"), NamedArgs);
+	}
+
+	TextBlock->SetColorAndOpacity(NewColor);
+	TextBlock->SetText(NewText);
+}
+
+#undef LOCTEXT_NAMESPACE
\ No newline at end of file
diff --git a/Engine/Source/Editor/CurveEditor/Private/Tree/SCurveEditorTreePin.cpp b/Engine/Source/Editor/CurveEditor/Private/Tree/SCurveEditorTreePin.cpp
--- a/Engine/Source/Editor/CurveEditor/Private/Tree/SCurveEditorTreePin.cpp
+++ b/Engine/Source/Editor/CurveEditor/Private/Tree/SCurveEditorTreePin.cpp
@@ -6,13 +6,15 @@
 
 #include "Widgets/Images/SImage.h"
 #include "Widgets/Input/SButton.h"
+#include "Widgets/Views/STableRow.h"
 
 #include "EditorStyleSet.h"
 
 
-void SCurveEditorTreePin::Construct(const FArguments& InArgs, TWeakPtr<FCurveEditor> InCurveEditor, FCurveEditorTreeItemID InTreeItemID)
+void SCurveEditorTreePin::Construct(const FArguments& InArgs, TWeakPtr<FCurveEditor> InCurveEditor, FCurveEditorTreeItemID InTreeItemID, const TSharedRef<ITableRow>& InTableRow)
 {
 	WeakCurveEditor = InCurveEditor;
+	WeakTableRow = InTableRow;
 	TreeItemID = InTreeItemID;
 
 	ChildSlot
@@ -103,8 +105,19 @@ bool SCurveEditorTreePin::IsPinnedRecursive(FCurveEditorTreeItemID InTreeItem, F
 
 EVisibility SCurveEditorTreePin::GetPinVisibility() const
 {
-	// ToDo: These should only be visible when hovering over the owning widget row.
-	return EVisibility::Visible;
+	TSharedPtr<FCurveEditor> CurveEditor = WeakCurveEditor.Pin();
+	TSharedPtr<ITableRow> Row = WeakTableRow.Pin();
+	TSharedPtr<SWidget> RowWidget = Row ? TSharedPtr<SWidget>(Row->AsWidget()) : nullptr;
+
+	if (RowWidget && RowWidget->IsHovered())
+	{
+		return EVisibility::Visible;
+	}
+	else if (CurveEditor && IsPinnedRecursive(TreeItemID, CurveEditor.Get()))
+	{
+		return EVisibility::Visible;
+	}
+	return EVisibility::Collapsed;
 }
 
 const FSlateBrush* SCurveEditorTreePin::GetPinBrush() const
diff --git a/Engine/Source/Editor/CurveEditor/Private/Tree/SCurveEditorTreeTextFilter.cpp b/Engine/Source/Editor/CurveEditor/Private/Tree/SCurveEditorTreeTextFilter.cpp
--- a/Engine/Source/Editor/CurveEditor/Private/Tree/SCurveEditorTreeTextFilter.cpp
+++ b/Engine/Source/Editor/CurveEditor/Private/Tree/SCurveEditorTreeTextFilter.cpp
@@ -12,31 +12,50 @@ void SCurveEditorTreeTextFilter::Construct(const FArguments& InArgs, TSharedPtr<
 {
 	WeakCurveEditor = CurveEditor;
 
+	CreateSearchBox();
+	CurveEditor->GetTree()->Events.OnFiltersChanged.AddSP(this, &SCurveEditorTreeTextFilter::OnTreeFilterListChanged);
+}
+
+void SCurveEditorTreeTextFilter::CreateSearchBox()
+{
 	ChildSlot
 	[
-		SNew(SSearchBox)
+		SAssignNew(SearchBox, SSearchBox)
 		.HintText(NSLOCTEXT("CurveEditor", "TextFilterHint", "Filter"))
 		.OnTextChanged(this, &SCurveEditorTreeTextFilter::OnFilterTextChanged)
 	];
 }
 
+void SCurveEditorTreeTextFilter::OnTreeFilterListChanged()
+{
+	TSharedPtr<FCurveEditor> CurveEditor = WeakCurveEditor.Pin();
+	if (CurveEditor && CurveEditor->GetTree()->FindFilterByType(ECurveEditorTreeFilterType::Text) == nullptr)
+	{
+		// If our filter has been removed externally, recreate the search box widget - this saves us having to manage a potentially
+		// recursive loop setting the search text and responding to the filter list being changed etc
+		Filter = nullptr;
+		CreateSearchBox();
+	}
+}
+
 void SCurveEditorTreeTextFilter::OnFilterTextChanged( const FText& FilterText )
 {
 	TSharedPtr<FCurveEditor> CurveEditor = WeakCurveEditor.Pin();
 	if (CurveEditor)
 	{
+		// Remove our binding to the OnFiltersChanged event since we know this will be broadcast
+		CurveEditor->GetTree()->Events.OnFiltersChanged.RemoveAll(this);
+
 		if (!Filter)
 		{
 			Filter = MakeShared<FCurveEditorTreeTextFilter>();
 		}
 
-		Filter->FilterTerms.Reset();
+		Filter->AssignFromText(FilterText.ToString());
 
-		static const bool bCullEmpty = true;
-		FilterText.ToString().ParseIntoArray(Filter->FilterTerms, TEXT(" "), bCullEmpty);
-
-		if (Filter->FilterTerms.Num() > 0)
+		if (!Filter->IsEmpty())
 		{
+			Filter->InputText = FilterText;
 			CurveEditor->GetTree()->AddFilter(Filter);
 		}
 		else
@@ -44,6 +63,8 @@ void SCurveEditorTreeTextFilter::OnFilterTextChanged( const FText& FilterText )
 			CurveEditor->GetTree()->RemoveFilter(Filter);
 		}
 
-		CurveEditor->GetTree()->RunFilters();
+		// Re-add our binding to the OnFiltersChanged now that we know it has been broadcast
+		CurveEditor->GetTree()->Events.OnFiltersChanged.AddSP(this, &SCurveEditorTreeTextFilter::OnTreeFilterListChanged);
 	}
-}
\ No newline at end of file
+}
+
diff --git a/Engine/Source/Editor/CurveEditor/Public/CurveEditor.h b/Engine/Source/Editor/CurveEditor/Public/CurveEditor.h
--- a/Engine/Source/Editor/CurveEditor/Public/CurveEditor.h
+++ b/Engine/Source/Editor/CurveEditor/Public/CurveEditor.h
@@ -293,11 +293,6 @@ public:
 	 */
 	void RemoveTreeItem(FCurveEditorTreeItemID ItemID);
 
-	/**
-	 * Set the current tree selection
-	 */
-	void SetDirectTreeSelection(TArray<FCurveEditorTreeItemID>&& TreeItems);
-
 	/**
 	 * Check whether this tree item is selected
 	 */
@@ -316,6 +311,14 @@ public:
 		return &Tree;
 	}
 
+	/**
+	 * Access the curve editor tree.
+	 */
+	const FCurveEditorTree* GetTree() const
+	{
+		return &Tree;
+	}
+
 	/**
 	 * Retrieve a serial number that is incremented any time a curve is added or removed
 	 */
diff --git a/Engine/Source/Editor/CurveEditor/Public/CurveEditorSettings.h b/Engine/Source/Editor/CurveEditor/Public/CurveEditorSettings.h
--- a/Engine/Source/Editor/CurveEditor/Public/CurveEditorSettings.h
+++ b/Engine/Source/Editor/CurveEditor/Public/CurveEditorSettings.h
@@ -37,8 +37,6 @@ class CURVEEDITOR_API UCurveEditorSettings : public UObject
 public:
 	GENERATED_BODY()
 
-	DECLARE_MULTICAST_DELEGATE(FOnCurveEditorCurveVisibilityChanged);
-
 	UCurveEditorSettings();
 
 	/** Gets whether or not the curve editor auto frames the selected curves. */
@@ -73,7 +71,4 @@ protected:
 
 	UPROPERTY( config, EditAnywhere, Category="Curve Editor")
 	ECurveEditorZoomPosition ZoomPosition;
-
-	/** Callback for when the Curve Editor Visibility setting changes. */
-	FOnCurveEditorCurveVisibilityChanged OnCurveEditorCurveVisibilityChanged;
 };
diff --git a/Engine/Source/Editor/CurveEditor/Public/Tree/CurveEditorTree.h b/Engine/Source/Editor/CurveEditor/Public/Tree/CurveEditorTree.h
--- a/Engine/Source/Editor/CurveEditor/Public/Tree/CurveEditorTree.h
+++ b/Engine/Source/Editor/CurveEditor/Public/Tree/CurveEditorTree.h
@@ -10,6 +10,8 @@
 #include "Delegates/Delegate.h"
 #include "CurveEditorTypes.h"
 
+enum class ECurveEditorTreeFilterType : uint32;
+
 struct ICurveEditorTreeItem;
 struct FCurveEditorTreeFilter;
 
@@ -33,54 +35,69 @@ enum class ECurveEditorTreeFilterState : uint8
 };
 
 /**
- * Scoped guard that will trigger the tree OnChanged event when all scoped guards have been exited
+ * Scoped guard that prevents the broadcast of tree events for the duration of its lifetime. Will trigger necessary events after the last remaining guard has been destroyed.
  */
-struct CURVEEDITOR_API FScopedCurveEditorTreeUpdateGuard
+struct CURVEEDITOR_API FScopedCurveEditorTreeEventGuard
 {
-	explicit FScopedCurveEditorTreeUpdateGuard(FCurveEditorTree* InTree);
-	~FScopedCurveEditorTreeUpdateGuard();
+	explicit FScopedCurveEditorTreeEventGuard(FCurveEditorTree* InTree);
+	~FScopedCurveEditorTreeEventGuard();
 
-	FScopedCurveEditorTreeUpdateGuard(FScopedCurveEditorTreeUpdateGuard&& RHS);
-	FScopedCurveEditorTreeUpdateGuard& operator=(FScopedCurveEditorTreeUpdateGuard&& RHS);
+	FScopedCurveEditorTreeEventGuard(FScopedCurveEditorTreeEventGuard&& RHS);
+	FScopedCurveEditorTreeEventGuard& operator=(FScopedCurveEditorTreeEventGuard&& RHS);
 
 private:
+	/** OnItemsChanged.SerialNumber cached on construction. Used to detect whether OnItemsChanged should be broadcast. */
+	uint32 CachedItemSerialNumber;
+	/** OnSelectionChanged.SerialNumber cached on construction. Used to detect whether OnSelectionChanged should be broadcast. */
+	uint32 CachedSelectionSerialNumber;
+	/** OnFiltersChanged.SerialNumber cached on construction. Used to detect whether OnFiltersChanged should be broadcast. */
+	uint32 CachedFiltersSerialNumber;
+	/** Pointer back to the owning tree */
 	FCurveEditorTree* Tree;
 };
 
+/**
+ * Generic multicast delegate that guards against re-entrancy for the curve editor tree
+ */
+class FCurveEditorTreeDelegate : public FSimpleMulticastDelegate
+{
+private:
+
+	friend FCurveEditorTree;
+	friend FScopedCurveEditorTreeEventGuard;
+
+	/* Broadcast this delegate provided it is not re-entrant */
+	void Broadcast();
+
+	/** Whether this delegate is currently broadcasting */
+	bool bBroadcasting = false;
+
+	/** Serial number that is incremented any time this delegate should be broadcast */
+	uint32 SerialNumber = 0;
+};
+
 /**
  * Struct that represents an event for when the tree has been changed.
- * This type carefully only allows FScopedCurveEditorTreeUpdateGuard to broadcast the event, and makes special checks for re-entrancy
+ * This type carefully only allows FScopedCurveEditorTreeEventGuard to broadcast the event, and makes special checks for re-entrancy
  */
-struct FCurveEditorOnChangedEvent
+struct FCurveEditorTreeEvents
 {
-	/**
-	 * @return true if the event is currently being broadcast
-	 */
-	bool IsBroadcastInProgress() const
-	{
-		return bIsBroadcastInProgress;
-	}
+	/** Event that is broadcast when the tree items container has changed */
+	FCurveEditorTreeDelegate OnItemsChanged;
 
-	/**
-	 * Bind a new handler to this OnChanged event
-	 */
-	FDelegateHandle Bind(FSimpleDelegate&& Handler) { return Delegate.Add(MoveTemp(Handler)); }
+	/** Event that is broadcast when the selection has changed */
+	FCurveEditorTreeDelegate OnSelectionChanged;
 
-	/**
-	 * Unbind a previously bound handler from this event
-	 */
-	void Unbind(FDelegateHandle Handle) { Delegate.Remove(Handle); }
+	/** Event that is broadcast when any kind of filtering has changed (ie active state, filters being added/removed etc) */
+	FCurveEditorTreeDelegate OnFiltersChanged;
 
 private:
 
-	friend FScopedCurveEditorTreeUpdateGuard;
+	friend FScopedCurveEditorTreeEventGuard;
 
-	/** Safety check to ensure that invocations of OnChangedEvent are never re-entrant */
-	bool bIsBroadcastInProgress = false;
-	/** Counter that is incremented for each living instance of FCurveEditorOnChangedEvent */
+	/** Counter that is incremented for each living instance of FScopedCurveEditorTreeEventGuard */
 	uint32 UpdateGuardCounter = 0;
-	/** The actual multi-cast delegate */
-	FSimpleMulticastDelegate Delegate;
+
 };
 
 /**
@@ -203,6 +220,18 @@ struct FCurveEditorFilterStates
 	void Reset()
 	{
 		FilterStates.Reset();
+		NumMatched = 0;
+		NumMatchedImplicitly = 0;
+	}
+
+	int32 GetNumMatched() const
+	{
+		return NumMatched;
+	}
+
+	int32 GetNumMatchedImplicitly() const
+	{
+		return NumMatchedImplicitly;
 	}
 
 	/**
@@ -226,7 +255,36 @@ struct FCurveEditorFilterStates
 	 */
 	void SetFilterState(FCurveEditorTreeItemID ItemID, ECurveEditorTreeFilterState NewState)
 	{
-		FilterStates.Add(ItemID, NewState);
+		const ECurveEditorTreeFilterState* Existing = FilterStates.Find(ItemID);
+		if (Existing)
+		{
+			if (*Existing  == ECurveEditorTreeFilterState::Match)
+			{
+				--NumMatched;
+			}
+			else if (*Existing != ECurveEditorTreeFilterState::NoMatch)
+			{
+				--NumMatchedImplicitly;
+			}
+		}
+
+		if (NewState == ECurveEditorTreeFilterState::NoMatch)
+		{
+			FilterStates.Remove(ItemID);
+		}
+		else
+		{
+			FilterStates.Add(ItemID, NewState);
+
+			if (NewState == ECurveEditorTreeFilterState::Match)
+			{
+				++NumMatched;
+			}
+			else
+			{
+				++NumMatchedImplicitly;
+			}
+		}
 	}
 
 	/**
@@ -258,6 +316,12 @@ private:
 	/** Whether filters should be active or not */
 	bool bIsActive = false;
 
+	/** The total number of nodes that have matched all the filters */
+	int32 NumMatched = 0;
+
+	/** The total number of nodes that have not matched the filters but have a parent or child that did */
+	int32 NumMatchedImplicitly = 0;
+
 	/** Filter state map. Items with no implicit or explicit filter state are not present */
 	TMap<FCurveEditorTreeItemID, ECurveEditorTreeFilterState> FilterStates;
 };
@@ -265,10 +329,13 @@ private:
 /** 
  * Complete implementation of a curve editor tree. Only really defines the hierarchy and selection states for tree items.
  */
-class FCurveEditorTree
+class CURVEEDITOR_API FCurveEditorTree
 {
 public:
 
+	/** Structure containing all the events for this tree */
+	FCurveEditorTreeEvents Events;
+
 	FCurveEditorTree();
 
 	/**
@@ -330,11 +397,29 @@ public:
 	 */
 	void RemoveFilter(TWeakPtr<FCurveEditorTreeFilter> FilterToRemove);
 
+	/**
+	 * Direct access to all current filters (potentially including expired ones)
+	 */
+	TArrayView<const TWeakPtr<FCurveEditorTreeFilter>> GetFilters() const;
+
+	/**
+	 * Clear all filters from the tree
+	 */
+	void ClearFilters();
+
+	/**
+	 * Attempt to locate a filter by its type
+	 *
+	 * @param Type The type of the filter to find
+	 * @return A pointer to the filter, or nullptr if one could not be found
+	 */
+	const FCurveEditorTreeFilter* FindFilterByType(ECurveEditorTreeFilterType Type) const;
+
 	/**
 	 * Inform this tree that the specified tree item IDs have been directly selected on the UI.
 	 * @note: This populates both implicit and explicit selection state for the supplied items and any children/parents
 	 */
-	void SetDirectSelection(TArray<FCurveEditorTreeItemID>&& TreeItems);
+	void SetDirectSelection(TArray<FCurveEditorTreeItemID>&& TreeItems, FCurveEditor* InCurveEditor);
 
 	/**
 	 * Access the selection state for this tree. Items that are neither implicitly or explicitly selected are not present in the map.
@@ -357,36 +442,12 @@ public:
 	ECurveEditorTreeFilterState GetFilterState(FCurveEditorTreeItemID InTreeItemID) const;
 
 	/**
-	 * Retrieve this tree's on-changed event
-	 */
-	FCurveEditorOnChangedEvent* OnChanged()
-	{
-		return &OnChangedEvent;
-	}
-
-	/**
-	 * Retrieve a scoped guard that will broadcast the on changed handlers for this tree when it goes out of scope (along with all other scoped guards on the stack)
-	 * Can be used to defer such broadcasts in situations where many changes are made to the tree at a time.
-	 */
-	FScopedCurveEditorTreeUpdateGuard ScopedUpdateGuard()
-	{
-		return FScopedCurveEditorTreeUpdateGuard(this);
-	}
-
-	/**
-	 * Add a handler for when this tree structure is changed in some way (items added/removed, tree filters changed etc)
-	 */
-	FDelegateHandle Bind_OnChanged(FSimpleDelegate&& Handler)
-	{
-		return OnChangedEvent.Bind(MoveTemp(Handler));
-	}
-
-	/**
-	 * Remove a handler for when this tree structure is changed in some way (items added/removed, tree filters changed etc)
+	 * Retrieve a scoped event guard that will block broadcast of events until the last guard on the stack goes out of scope
+	 * Can be used to defer broadcasts in situations where many changes are made to the tree at a time.
 	 */
-	void Unbind_OnChanged(FDelegateHandle Handle)
+	FScopedCurveEditorTreeEventGuard ScopedEventGuard()
 	{
-		return OnChangedEvent.Unbind(Handle);
+		return FScopedCurveEditorTreeEventGuard(this);
 	}
 
 private:
@@ -398,23 +459,20 @@ private:
 	 * Run the specified filters over the specified items and their recursive children, storing the results in this instance's FilterStates struct.
 	 *
 	 * @param FilterPtrs     Array of non-null pointers to filters to use. Items are considered matched if they match any filter in this array.
-	 * @param Items          Array item IDs to filter
+	 * @param ItemsToFilter  Array item IDs to filter
 	 * @param InheritedState The filter state for each item to receive if it does not directly match a filter (either ECurveEditorTreeFilterState::NoMatch or ECurveEditorTreeFilterState::InheritedChild)
 	 * @return Whether any of the items or any their recursive children matched any filter
 	 */
-	bool FilterSpecificItems(TArrayView<FCurveEditorTreeFilter* const> FilterPtrs, TArrayView<const FCurveEditorTreeItemID> Items, ECurveEditorTreeFilterState InheritedState);
+	bool PerformFilterPass(TArrayView<const FCurveEditorTreeFilter* const> FilterPtrs, TArrayView<const FCurveEditorTreeItemID> ItemsToFilter, ECurveEditorTreeFilterState InheritedState);
 
 	/** Incrementing ID for the next tree item to be created */
 	FCurveEditorTreeItemID NextTreeItemID;
 
-	/** Container housing the machinery required for deferred broadcast of changes to the tree */
-	FCurveEditorOnChangedEvent OnChangedEvent;
-
 	/** Map of all tree items by their ID */
 	TMap<FCurveEditorTreeItemID, FCurveEditorTreeItem> Items;
 
-	/** Map of all tree items by their ID */
-	TArray<TWeakPtr<FCurveEditorTreeFilter>> Filters;
+	/** Array of all filters that are currently active on this tree */
+	TArray<TWeakPtr<FCurveEditorTreeFilter>> WeakFilters;
 
 	/** Hierarchical information for the tree */
 	FSortedCurveEditorTreeItems RootItems;
diff --git a/Engine/Source/Editor/CurveEditor/Public/Tree/CurveEditorTreeFilter.h b/Engine/Source/Editor/CurveEditor/Public/Tree/CurveEditorTreeFilter.h
--- a/Engine/Source/Editor/CurveEditor/Public/Tree/CurveEditorTreeFilter.h
+++ b/Engine/Source/Editor/CurveEditor/Public/Tree/CurveEditorTreeFilter.h
@@ -5,6 +5,7 @@
 #include "Templates/SharedPointer.h"
 #include "Misc/EnumClassFlags.h"
 #include "Containers/UnrealString.h"
+#include "Internationalization/Text.h"
 
 class FCurveEditorTree;
 struct FCurveEditorTreeItem;
@@ -24,12 +25,16 @@ enum class ECurveEditorTreeFilterType : uint32
 	First = Text,
 };
 
-
-
+/**
+ * Base class for all filters that can be applied to a curve editor tree.
+ * Filters are identifyable through their type (see GetType()), which is a pre-registered static value retrieved through RegisterFilterType(). Client types can inspect this for supported types, and static_cast<> for implementation details.
+ * Filters also contain a pass, which is used to separate filters into separate passes. Filters applied within the same pass are matched as a boolean OR, filters in different passes are matched as a boolean AND.
+ */
 struct CURVEEDITOR_API FCurveEditorTreeFilter
 {
-	explicit FCurveEditorTreeFilter(ECurveEditorTreeFilterType InFilterType)
+	FCurveEditorTreeFilter(ECurveEditorTreeFilterType InFilterType, int32 InFilterPass)
 		: FilterType(InFilterType)
+		, FilterPass(InFilterPass)
 	{}
 
 	virtual ~FCurveEditorTreeFilter() {}
@@ -42,6 +47,14 @@ struct CURVEEDITOR_API FCurveEditorTreeFilter
 		return FilterType;
 	}
 
+	/**
+	 * @return The pass index that should be used when applying this filter
+	 */
+	int32 GetFilterPass() const
+	{
+		return FilterPass;
+	}
+
 public:
 
 	/**
@@ -56,32 +69,71 @@ protected:
 
 	/** The static type of this filter as retrieved by RegisterFilterType */
 	ECurveEditorTreeFilterType FilterType;
+
+	/** Defines which pass this filter should be applied in */
+	int32 FilterPass;
+};
+
+/** A specific text token (containing neither spaces nor .) */
+struct FCurveEditorTreeTextFilterToken
+{
+	FString Token;
+
+	/** Match this token against a string */
+	bool Match(const TCHAR* InString) const
+	{
+		return FCString::Stristr(InString, *Token) != nullptr;
+	}
+};
+
+/** A text filter term containing >= 1 sparate tokens ordered from child to parent */
+struct FCurveEditorTreeTextFilterTerm
+{
+	TArray<FCurveEditorTreeTextFilterToken, TInlineAllocator<1>> ChildToParentTokens;
 };
 
 /**
  * Built-in text filter of type ECurveEditorTreeFilterType::Text. Filter terms are applied as a case-insensitive boolean OR substring match.
  */
-struct FCurveEditorTreeTextFilter : FCurveEditorTreeFilter
+struct CURVEEDITOR_API FCurveEditorTreeTextFilter : FCurveEditorTreeFilter
 {
-	/** Array of case-insensitive terms to find within tree items. */
-	TArray<FString> FilterTerms;
+	/** Original input filter text. */
+	FText InputText;
+
+	/** Default pass for text filters. */
+	static const int32 DefaultPass = 1000;
 
+	/** Default constructor */
 	FCurveEditorTreeTextFilter()
-		: FCurveEditorTreeFilter(ECurveEditorTreeFilterType::Text)
+		: FCurveEditorTreeFilter(ECurveEditorTreeFilterType::Text, DefaultPass)
 	{}
 
 	/**
-	 * Check whether the supplied string matches any of the terms
+	 * Assign a new filter string to this filter, resetting any previous filter terms
+	 *
+	 * @param FilterString   The new filter string that should be parsed into this filter
 	 */
-	bool Match(const TCHAR* InString) const
+	void AssignFromText(const FString& FilterString);
+
+	/**
+	 * Check whether this filter is empty (ie, has no filter terms)
+	 */
+	bool IsEmpty() const
+	{
+		return ChildToParentFilterTerms.Num() == 0;
+	}
+
+	/**
+	 * Access all the filter terms contained within this filter. Each term comprises an array of text tokens that must match from child to parent.
+	 * Valid to call if IsEmpty() == true, but the resulting array view will be empty.
+	 */
+	TArrayView<const FCurveEditorTreeTextFilterTerm> GetTerms() const
 	{
-		for (const FString& Term : FilterTerms)
-		{
-			if (FCString::Stristr(InString, *Term) != nullptr)
-			{
-				return true;
-			}
-		}
-		return false;
+		return ChildToParentFilterTerms;
 	}
+
+private:
+
+	/** Arrays of case-insensitive terms to find within tree items. Each term is an array of strings that must be matched from child to parent. */
+	TArray<FCurveEditorTreeTextFilterTerm, TInlineAllocator<4>> ChildToParentFilterTerms;
 };
\ No newline at end of file
diff --git a/Engine/Source/Editor/CurveEditor/Public/Tree/CurveEditorTreeTraits.h b/Engine/Source/Editor/CurveEditor/Public/Tree/CurveEditorTreeTraits.h
--- a/Engine/Source/Editor/CurveEditor/Public/Tree/CurveEditorTreeTraits.h
+++ b/Engine/Source/Editor/CurveEditor/Public/Tree/CurveEditorTreeTraits.h
@@ -2,7 +2,7 @@
 
 #pragma once
 
-#include "Widgets/DeclarativeSyntaxSupport.h"
+#include "Framework/Views/TableViewTypeTraits.h"
 #include "CurveEditorTypes.h"
 
 template<>
@@ -19,7 +19,8 @@ public:
 	using MapKeyFuncsSparse = TDefaultMapHashableKeyFuncs<FCurveEditorTreeItemID, FSparseItemInfo, false>;
 	using SetKeyFuncs = DefaultKeyFuncs<FCurveEditorTreeItemID>;
 
-	static void AddReferencedObjects(FReferenceCollector&, TArray<FCurveEditorTreeItemID>&, TSet<FCurveEditorTreeItemID>&) {}
+	template<typename U>
+	static void AddReferencedObjects(FReferenceCollector&, TArray<FCurveEditorTreeItemID>&, TSet<FCurveEditorTreeItemID>&, TMap< const U*, FCurveEditorTreeItemID >&) {}
 
 	static bool IsPtrValid(NullableType InPtr)
 	{
diff --git a/Engine/Source/Editor/CurveEditor/Public/Tree/ICurveEditorTreeItem.h b/Engine/Source/Editor/CurveEditor/Public/Tree/ICurveEditorTreeItem.h
--- a/Engine/Source/Editor/CurveEditor/Public/Tree/ICurveEditorTreeItem.h
+++ b/Engine/Source/Editor/CurveEditor/Public/Tree/ICurveEditorTreeItem.h
@@ -13,9 +13,11 @@ struct FCurveEditorTreeFilter;
 
 class FName;
 class SWidget;
+class ITableRow;
 class FCurveModel;
 class FCurveEditor;
 
+
 /**
  * Optional implementation interface for any tree item to be shown on the curve editor tree.
  */
@@ -29,9 +31,10 @@ struct CURVEEDITOR_API ICurveEditorTreeItem
 	 * @param InColumnName     The name of the column to generate widget content for. See FColumnNames for valid names.
 	 * @param InCurveEditor    Weak pointer to the curve editor instance. Persistent TSharedPtrs should not be held to this pointer.
 	 * @param InTreeItemID     The ID of the tree item that this interface is assigned to
+	 * @param InTableRow       The table row that will house the widget content
 	 * @return (Optional) Widget content for this column, or nullptr if none is necessary for this column.
 	 */
-	virtual TSharedPtr<SWidget> GenerateCurveEditorTreeWidget(const FName& InColumnName, TWeakPtr<FCurveEditor> InCurveEditor, FCurveEditorTreeItemID InTreeItemID) = 0;
+	virtual TSharedPtr<SWidget> GenerateCurveEditorTreeWidget(const FName& InColumnName, TWeakPtr<FCurveEditor> InCurveEditor, FCurveEditorTreeItemID InTreeItemID, const TSharedRef<ITableRow>& InTableRow) = 0;
 
 	/**
 	 * Populate the specified array with curve models that are represented by this tree item
diff --git a/Engine/Source/Editor/CurveEditor/Public/Tree/SCurveEditorTree.h b/Engine/Source/Editor/CurveEditor/Public/Tree/SCurveEditorTree.h
--- a/Engine/Source/Editor/CurveEditor/Public/Tree/SCurveEditorTree.h
+++ b/Engine/Source/Editor/CurveEditor/Public/Tree/SCurveEditorTree.h
@@ -3,16 +3,16 @@
 #pragma once
 
 #include "Widgets/DeclarativeSyntaxSupport.h"
-#include "Widgets/SCompoundWidget.h"
+#include "Widgets/Views/STreeView.h"
+#include "CurveEditorTreeTraits.h"
 #include "CurveEditorTypes.h"
 
 class FCurveEditor;
 class ITableRow;
 class SHeaderRow;
 class STableViewBase;
-class SCurveEditorTreeView;
 
-class CURVEEDITOR_API SCurveEditorTree : public SCompoundWidget
+class CURVEEDITOR_API SCurveEditorTree : public STreeView<FCurveEditorTreeItemID>
 {
 public:
 
@@ -49,5 +49,4 @@ private:
 	TSharedPtr<FCurveEditor> CurveEditor;
 
 	TSharedPtr<SHeaderRow> HeaderRow;
-	TSharedPtr<SCurveEditorTreeView> TreeView;
 };
\ No newline at end of file
diff --git a/Engine/Source/Editor/CurveEditor/Public/Tree/SCurveEditorTreeFilterStatusBar.h b/Engine/Source/Editor/CurveEditor/Public/Tree/SCurveEditorTreeFilterStatusBar.h
new file mode 100644
--- /dev/null
+++ b/Engine/Source/Editor/CurveEditor/Public/Tree/SCurveEditorTreeFilterStatusBar.h
@@ -0,0 +1,36 @@
+// Copyright 1998-2019 Epic Games, Inc. All Rights Reserved.
+
+#pragma once
+
+#include "Widgets/DeclarativeSyntaxSupport.h"
+#include "Widgets/SCompoundWidget.h"
+
+class STextBlock;
+class FCurveEditor;
+
+/** Widget that shows the metrics for the current tree filter in the curve editor (in the form "Showing {0} of {1} items ({2} selected)") */
+class CURVEEDITOR_API SCurveEditorTreeFilterStatusBar : public SCompoundWidget
+{
+public:
+	SLATE_BEGIN_ARGS(SCurveEditorTreeFilterStatusBar){}
+	SLATE_END_ARGS()
+
+	/** Construct the status bar */
+	void Construct(const FArguments& InArgs, TSharedPtr<FCurveEditor> InCurveEditor);
+
+private:
+
+	/** Request the visibility of the clear hyperlink widget based on whether there's a filter active or not */
+	EVisibility GetVisibilityFromFilter() const;
+
+	/** Update the filter text to represent the current filter states in the tree */
+	void UpdateText();
+
+	/** Clear all the filters in the tree - called as a result of the user clicking on the 'clear' hyperlink */
+	void ClearFilters();
+
+private:
+
+	TWeakPtr<FCurveEditor> WeakCurveEditor;
+	TSharedPtr<STextBlock> TextBlock;
+};
\ No newline at end of file
diff --git a/Engine/Source/Editor/CurveEditor/Public/Tree/SCurveEditorTreePin.h b/Engine/Source/Editor/CurveEditor/Public/Tree/SCurveEditorTreePin.h
--- a/Engine/Source/Editor/CurveEditor/Public/Tree/SCurveEditorTreePin.h
+++ b/Engine/Source/Editor/CurveEditor/Public/Tree/SCurveEditorTreePin.h
@@ -6,8 +6,10 @@
 #include "Widgets/SCompoundWidget.h"
 
 #include "CurveEditorTypes.h"
+#include "CurveEditorTreeTraits.h"
 
 class FCurveEditor;
+class ITableRow;
 
 class CURVEEDITOR_API SCurveEditorTreePin : public SCompoundWidget
 {
@@ -15,7 +17,7 @@ public:
 	SLATE_BEGIN_ARGS(SCurveEditorTreePin){}
 	SLATE_END_ARGS()
 
-	void Construct(const FArguments& InArgs, TWeakPtr<FCurveEditor> InCurveEditor, FCurveEditorTreeItemID InTreeItemID);
+	void Construct(const FArguments& InArgs, TWeakPtr<FCurveEditor> InCurveEditor, FCurveEditorTreeItemID InTreeItemID, const TSharedRef<ITableRow>& InTableRow);
 
 private:
 
@@ -32,6 +34,7 @@ private:
 
 private:
 
+	TWeakPtr<ITableRow> WeakTableRow;
 	TWeakPtr<FCurveEditor> WeakCurveEditor;
 	FCurveEditorTreeItemID TreeItemID;
 };
\ No newline at end of file
diff --git a/Engine/Source/Editor/CurveEditor/Public/Tree/SCurveEditorTreeTextFilter.h b/Engine/Source/Editor/CurveEditor/Public/Tree/SCurveEditorTreeTextFilter.h
--- a/Engine/Source/Editor/CurveEditor/Public/Tree/SCurveEditorTreeTextFilter.h
+++ b/Engine/Source/Editor/CurveEditor/Public/Tree/SCurveEditorTreeTextFilter.h
@@ -5,6 +5,7 @@
 #include "Widgets/DeclarativeSyntaxSupport.h"
 #include "Widgets/SCompoundWidget.h"
 
+class SSearchBox;
 class FCurveEditor;
 struct FCurveEditorTreeTextFilter;
 
@@ -19,8 +20,11 @@ public:
 
 private:
 
+	void CreateSearchBox();
+	void OnTreeFilterListChanged();
 	void OnFilterTextChanged(const FText& FilterText);
 
+	TSharedPtr<SSearchBox> SearchBox;
 	TSharedPtr<FCurveEditorTreeTextFilter> Filter;
 	TWeakPtr<FCurveEditor> WeakCurveEditor;
 };
\ No newline at end of file
diff --git a/Engine/Source/Editor/Sequencer/Private/DisplayNodes/SequencerDisplayNode.cpp b/Engine/Source/Editor/Sequencer/Private/DisplayNodes/SequencerDisplayNode.cpp
--- a/Engine/Source/Editor/Sequencer/Private/DisplayNodes/SequencerDisplayNode.cpp
+++ b/Engine/Source/Editor/Sequencer/Private/DisplayNodes/SequencerDisplayNode.cpp
@@ -25,6 +25,7 @@
 #include "SSequencerSectionAreaView.h"
 #include "CommonMovieSceneTools.h"
 #include "Framework/Commands/GenericCommands.h"
+#include "CurveEditor.h"
 #include "Tree/SCurveEditorTreePin.h"
 #include "Tree/CurveEditorTreeFilter.h"
 #include "ScopedTransaction.h"
@@ -96,6 +97,46 @@ namespace SequencerNodeConstants
 		// New nodes get added to the correctly sorted position by default
 		return SortChildrenWithBias(A, B, SequencerNodeConstants::DefaultSortBias);
 	}
+
+	static bool NodeMatchesTextFilterTerm(TSharedPtr<const FSequencerDisplayNode> Node, const FCurveEditorTreeTextFilterTerm& Term)
+	{
+		bool bMatched = false;
+
+		for (const FCurveEditorTreeTextFilterToken& Token : Term.ChildToParentTokens)
+		{
+			if (!Node)
+			{
+				// No match - ran out of parents
+				return false;
+			}
+			else if (!Token.Match(*Node->GetDisplayName().ToString()))
+			{
+				return false;
+			}
+
+			bMatched = true;
+			Node = Node->GetParent();
+		}
+
+		return bMatched;
+	}
+
+	FText GetCurveEditorHighlightText(TWeakPtr<FCurveEditor> InCurveEditor)
+	{
+		TSharedPtr<FCurveEditor> PinnedCurveEditor = InCurveEditor.Pin();
+		if (!PinnedCurveEditor)
+		{
+			return FText::GetEmpty();
+		}
+
+		const FCurveEditorTreeFilter* Filter = PinnedCurveEditor->GetTree()->FindFilterByType(ECurveEditorTreeFilterType::Text);
+		if (Filter)
+		{
+			return static_cast<const FCurveEditorTreeTextFilter*>(Filter)->InputText;
+		}
+
+		return FText::GetEmpty();
+	}
 }
 
 struct FNameAndSignature
@@ -1074,7 +1115,7 @@ bool FSequencerDisplayNode::HandleContextMenuRenameNodeCanExecute() const
 }
 
 
-TSharedPtr<SWidget> FSequencerDisplayNode::GenerateCurveEditorTreeWidget(const FName& InColumnName, TWeakPtr<FCurveEditor> InCurveEditor, FCurveEditorTreeItemID InTreeItemID)
+TSharedPtr<SWidget> FSequencerDisplayNode::GenerateCurveEditorTreeWidget(const FName& InColumnName, TWeakPtr<FCurveEditor> InCurveEditor, FCurveEditorTreeItemID InTreeItemID, const TSharedRef<ITableRow>& TableRow)
 {
 	if (InColumnName == ColumnNames.Label)
 	{
@@ -1112,15 +1153,17 @@ TSharedPtr<SWidget> FSequencerDisplayNode::GenerateCurveEditorTreeWidget(const F
 
 			+ SHorizontalBox::Slot()
 			.VAlign(VAlign_Center)
+			.Padding(FMargin(0.f, 4.f, 0.f, 4.f))
 			[
 				SNew(STextBlock)
 				.Text(this, &FSequencerDisplayNode::GetDisplayName)
+				.HighlightText_Static(SequencerNodeConstants::GetCurveEditorHighlightText, InCurveEditor)
 				.ToolTipText(this, &FSequencerDisplayNode::GetDisplayNameToolTipText)
 			];
 	}
 	else if (InColumnName == ColumnNames.PinHeader)
 	{
-		return SNew(SCurveEditorTreePin, InCurveEditor, InTreeItemID);
+		return SNew(SCurveEditorTreePin, InCurveEditor, InTreeItemID, TableRow);
 	}
 
 	return nullptr;
@@ -1135,7 +1178,21 @@ bool FSequencerDisplayNode::PassesFilter(const FCurveEditorTreeFilter* InFilter)
 	if (InFilter->GetType() == ECurveEditorTreeFilterType::Text)
 	{
 		const FCurveEditorTreeTextFilter* Filter = static_cast<const FCurveEditorTreeTextFilter*>(InFilter);
-		return Filter->Match(*GetDisplayName().ToString());
+
+		TSharedRef<const FSequencerDisplayNode> This = AsShared();
+		for (const FCurveEditorTreeTextFilterTerm& Term : Filter->GetTerms())
+		{
+			if (SequencerNodeConstants::NodeMatchesTextFilterTerm(This, Term))
+			{
+				return true;
+			}
+		}
+
+		return false;
+	}
+	else if (InFilter->GetType() == ISequencerModule::GetSequencerSelectionFilterType())
+	{
+		return GetSequencer().GetSelection().GetSelectedOutlinerNodes().Contains(ConstCastSharedRef<FSequencerDisplayNode>(AsShared()));
 	}
 	return false;
 }
diff --git a/Engine/Source/Editor/Sequencer/Private/DisplayNodes/SequencerDisplayNode.h b/Engine/Source/Editor/Sequencer/Private/DisplayNodes/SequencerDisplayNode.h
--- a/Engine/Source/Editor/Sequencer/Private/DisplayNodes/SequencerDisplayNode.h
+++ b/Engine/Source/Editor/Sequencer/Private/DisplayNodes/SequencerDisplayNode.h
@@ -274,7 +274,7 @@ public:
 	virtual FText GetIconToolTipText() const;
 
 	// ICurveEditorTreeItem interface
-	virtual TSharedPtr<SWidget> GenerateCurveEditorTreeWidget(const FName& InColumnName, TWeakPtr<FCurveEditor> InCurveEditor, FCurveEditorTreeItemID InTreeItemID) override;
+	virtual TSharedPtr<SWidget> GenerateCurveEditorTreeWidget(const FName& InColumnName, TWeakPtr<FCurveEditor> InCurveEditor, FCurveEditorTreeItemID InTreeItemID, const TSharedRef<ITableRow>& InTableRow) override;
 	virtual void CreateCurveModels(TArray<TUniquePtr<FCurveModel>>& OutCurveModels) override;
 	virtual bool PassesFilter(const FCurveEditorTreeFilter* InFilter) const override;
 
diff --git a/Engine/Source/Editor/Sequencer/Private/SSequencer.cpp b/Engine/Source/Editor/Sequencer/Private/SSequencer.cpp
--- a/Engine/Source/Editor/Sequencer/Private/SSequencer.cpp
+++ b/Engine/Source/Editor/Sequencer/Private/SSequencer.cpp
@@ -72,7 +72,9 @@
 #include "Camera/CameraActor.h"
 #include "SCurveEditorPanel.h"
 #include "Tree/SCurveEditorTree.h"
+#include "Tree/CurveEditorTreeFilter.h"
 #include "Tree/SCurveEditorTreeTextFilter.h"
+#include "Tree/SCurveEditorTreeFilterStatusBar.h"
 #include "SCurveKeyDetailPanel.h"
 #include "MovieSceneTimeHelpers.h"
 #include "FrameNumberNumericInterface.h"
@@ -414,6 +416,7 @@ void SSequencer::Construct(const FArguments& InArgs, TSharedRef<FSequencer> InSe
 	// Some areas that use Sequencer don't use our curve editor. In this case no button is shown on the UI.
 	if (InSequencer->GetToolkitHost().IsValid())
 	{
+		CurveEditorTree = SNew(SCurveEditorTree, InSequencer->GetCurveEditor());
 		TSharedRef<SCurveEditorPanel> CurveEditorWidget = SNew(SCurveEditorPanel, InSequencer->GetCurveEditor().ToSharedRef())
 			// Grid lines match the color specified in FSequencerTimeSliderController::OnPaintViewArea
 			.GridLineTint(FLinearColor(0.f, 0.f, 0.f, 0.3f))
@@ -432,7 +435,16 @@ void SSequencer::Construct(const FArguments& InArgs, TSharedRef<FSequencer> InSe
 
 				+ SVerticalBox::Slot()
 				[
-					SNew(SCurveEditorTree, InSequencer->GetCurveEditor())
+					SNew(SScrollBorder, CurveEditorTree.ToSharedRef())
+					[
+						CurveEditorTree.ToSharedRef()
+					]
+				]
+
+				+ SVerticalBox::Slot()
+				.AutoHeight()
+				[
+					SNew(SCurveEditorTreeFilterStatusBar, InSequencer->GetCurveEditor())
 				]
 
 				+ SVerticalBox::Slot()
@@ -443,13 +455,13 @@ void SSequencer::Construct(const FArguments& InArgs, TSharedRef<FSequencer> InSe
 				]
 			];
 
-			// Register an instanced custom property type layout to handle converting FFrameNumber from Tick Resolution to Display Rate.
-			CurveEditorWidget->GetKeyDetailsView()->GetPropertyRowGenerator()->RegisterInstancedCustomPropertyTypeLayout("FrameNumber", FOnGetPropertyTypeCustomizationInstance::CreateStatic(CreateFrameNumberCustomization, SequencerPtr));
-			TAttribute<bool> IsEnabledAttribute = TAttribute<bool>::Create(TAttribute<bool>::FGetter::CreateSP(this, &SSequencer::GetIsCurveEditorEnabled));
+		// Register an instanced custom property type layout to handle converting FFrameNumber from Tick Resolution to Display Rate.
+		CurveEditorWidget->GetKeyDetailsView()->GetPropertyRowGenerator()->RegisterInstancedCustomPropertyTypeLayout("FrameNumber", FOnGetPropertyTypeCustomizationInstance::CreateStatic(CreateFrameNumberCustomization, SequencerPtr));
+		TAttribute<bool> IsEnabledAttribute = TAttribute<bool>::Create(TAttribute<bool>::FGetter::CreateSP(this, &SSequencer::GetIsCurveEditorEnabled));
 	
-			CurveEditorPanel = SNew(SSequencerCurveEditor, CurveEditorWidget);
-			CurveEditorPanel->SetEnabled(IsEnabledAttribute);
-			CurveEditorWidget->SetEnabled(IsEnabledAttribute);
+		CurveEditorPanel = SNew(SSequencerCurveEditor, CurveEditorWidget);
+		CurveEditorPanel->SetEnabled(IsEnabledAttribute);
+		CurveEditorWidget->SetEnabled(IsEnabledAttribute);
 
 		// Check to see if the tab is already opened due to the saved window layout.
 		TSharedPtr<SDockTab> ExistingCurveEditorTab = InSequencer->GetToolkitHost()->GetTabManager()->FindExistingLiveTab(FTabId(SSequencer::CurveEditorTabName));
@@ -832,6 +844,7 @@ void SSequencer::Construct(const FArguments& InArgs, TSharedRef<FSequencer> InSe
 
 	InSequencer->GetSelection().GetOnKeySelectionChanged().AddSP(this, &SSequencer::HandleKeySelectionChanged);
 	InSequencer->GetSelection().GetOnSectionSelectionChanged().AddSP(this, &SSequencer::HandleSectionSelectionChanged);
+	InSequencer->GetSelection().GetOnOutlinerNodeSelectionChanged().AddSP(this, &SSequencer::HandleOutlinerNodeSelectionChanged);
 
 	ResetBreadcrumbs();
 }
@@ -978,6 +991,51 @@ void SSequencer::HandleSectionSelectionChanged()
 {
 }
 
+
+void SSequencer::HandleOutlinerNodeSelectionChanged()
+{
+	TSharedPtr<FSequencer> Sequencer = SequencerPtr.Pin();
+	if (!Sequencer.IsValid())
+	{
+		return;
+	}
+
+	const TSet<TSharedRef<FSequencerDisplayNode>>& SelectedDisplayNodes = Sequencer->GetSelection().GetSelectedOutlinerNodes();
+
+	// If we're isolating to the selection and there is one, add the filter
+	if (Settings->ShouldIsolateToCurveEditorSelection() && SelectedDisplayNodes.Num() != 0)
+	{
+		if (!SequencerSelectionCurveEditorFilter)
+		{
+			static const int32 FilterPass = -1000;
+			SequencerSelectionCurveEditorFilter = MakeShared<FCurveEditorTreeFilter>(ISequencerModule::GetSequencerSelectionFilterType(), FilterPass);
+		}
+		Sequencer->GetCurveEditor()->GetTree()->AddFilter(SequencerSelectionCurveEditorFilter);
+	}
+	// If we're not isolating to the selection (or there is no selection) remove the filter
+	else if (SequencerSelectionCurveEditorFilter)
+	{
+		Sequencer->GetCurveEditor()->GetTree()->RemoveFilter(SequencerSelectionCurveEditorFilter);
+		SequencerSelectionCurveEditorFilter = nullptr;
+	}
+
+	if (Settings->ShouldSyncCurveEditorSelection())
+	{
+		TSharedRef<FSequencerNodeTree> NodeTree = Sequencer->GetNodeTree();
+
+		// Clear the tree selection
+		CurveEditorTree->ClearSelection();
+		for (TSharedRef<FSequencerDisplayNode> Node : SelectedDisplayNodes)
+		{
+			FCurveEditorTreeItemID CurveEditorTreeItem = NodeTree->FindCurveEditorTreeItem(Node);
+			if (CurveEditorTreeItem != FCurveEditorTreeItemID::Invalid())
+			{
+				CurveEditorTree->SetItemSelection(CurveEditorTreeItem, true);
+			}
+		}
+	}
+}
+
 TSharedRef<SWidget> SSequencer::MakeAddButton()
 {
 	return SNew(SComboButton)
diff --git a/Engine/Source/Editor/Sequencer/Private/SSequencer.h b/Engine/Source/Editor/Sequencer/Private/SSequencer.h
--- a/Engine/Source/Editor/Sequencer/Private/SSequencer.h
+++ b/Engine/Source/Editor/Sequencer/Private/SSequencer.h
@@ -25,6 +25,7 @@ class FSequencerTimeSliderController;
 class FVirtualTrackArea;
 class ISequencerEditTool;
 class SSequencerLabelBrowser;
+class SCurveEditorTree;
 class SSequencerTrackArea;
 class SSequencerTrackOutliner;
 class SSequencerTransformBox;
@@ -322,6 +323,9 @@ private:
 	/** Handles section selection changes. */
 	void HandleSectionSelectionChanged();
 
+	/** Handles changes to the selected outliner nodes. */
+	void HandleOutlinerNodeSelectionChanged();
+
 	/** Syncs the current node selection to the curve editor. */
 	void SyncCurveEditorToSelection();
 
@@ -524,6 +528,12 @@ private:
 	/** Outliner widget */
 	TSharedPtr<SSequencerTrackOutliner> TrackOutliner;
 
+	/** Curve editor tree widget */
+	TSharedPtr<SCurveEditorTree> CurveEditorTree;
+
+	/** Curve editor filter that shows only the selected nodes */
+	TSharedPtr<FCurveEditorTreeFilter> SequencerSelectionCurveEditorFilter;
+
 	/** The breadcrumb trail widget for this sequencer */
 	TSharedPtr<SBreadcrumbTrail<FSequencerBreadcrumb>> BreadcrumbTrail;
 
diff --git a/Engine/Source/Editor/Sequencer/Private/SequencerModule.cpp b/Engine/Source/Editor/Sequencer/Private/SequencerModule.cpp
--- a/Engine/Source/Editor/Sequencer/Private/SequencerModule.cpp
+++ b/Engine/Source/Editor/Sequencer/Private/SequencerModule.cpp
@@ -13,6 +13,7 @@
 #include "SequencerEdMode.h"
 #include "SequencerObjectChangeListener.h"
 #include "IDetailKeyframeHandler.h"
+#include "Tree/CurveEditorTreeFilter.h"
 #include "AnimatedPropertyKey.h"
 
 #define LOCTEXT_NAMESPACE "SequencerEditor"
@@ -22,6 +23,12 @@ ISequencerModule::~ISequencerModule()
 {
 }
 
+ECurveEditorTreeFilterType ISequencerModule::GetSequencerSelectionFilterType()
+{
+	static ECurveEditorTreeFilterType FilterType = FCurveEditorTreeFilter::RegisterFilterType();
+	return FilterType;
+}
+
 /**
  * SequencerModule implementation (private)
  */
diff --git a/Engine/Source/Editor/Sequencer/Private/SequencerNodeTree.cpp b/Engine/Source/Editor/Sequencer/Private/SequencerNodeTree.cpp
--- a/Engine/Source/Editor/Sequencer/Private/SequencerNodeTree.cpp
+++ b/Engine/Source/Editor/Sequencer/Private/SequencerNodeTree.cpp
@@ -844,7 +844,7 @@ void FSequencerNodeTree::UpdateCurveEditorTree()
 	FCurveEditor* CurveEditor = Sequencer.GetCurveEditor().Get();
 
 	// Guard against multiple broadcasts here and defer them until the end of this function
-	FScopedCurveEditorTreeUpdateGuard ScopedUpdateGuard = CurveEditor->GetTree()->ScopedUpdateGuard();
+	FScopedCurveEditorTreeEventGuard ScopedEventGuard = CurveEditor->GetTree()->ScopedEventGuard();
 
 	auto Traverse_AddToCurveEditor = [this, CurveEditor](FSequencerDisplayNode& InNode)
 	{
@@ -909,4 +909,10 @@ FCurveEditorTreeItemID FSequencerNodeTree::AddToCurveEditor(TSharedRef<FSequence
 
 	CurveEditorTreeItemIDs.Add(DisplayNode, NewItem->GetID());
 	return NewItem->GetID();
+}
+
+FCurveEditorTreeItemID FSequencerNodeTree::FindCurveEditorTreeItem(TSharedRef<FSequencerDisplayNode> DisplayNode) const
+{
+	const FCurveEditorTreeItemID* FoundID = CurveEditorTreeItemIDs.Find(DisplayNode);
+	return FoundID ? *FoundID : FCurveEditorTreeItemID::Invalid();
 }
\ No newline at end of file
diff --git a/Engine/Source/Editor/Sequencer/Private/SequencerNodeTree.h b/Engine/Source/Editor/Sequencer/Private/SequencerNodeTree.h
--- a/Engine/Source/Editor/Sequencer/Private/SequencerNodeTree.h
+++ b/Engine/Source/Editor/Sequencer/Private/SequencerNodeTree.h
@@ -139,6 +139,14 @@ public:
 	/** Sorts all nodes and their descendants by category then alphabetically.*/
 	void SortAllNodesAndDescendants();
 
+	/**
+	 * Attempt to find a curve editor tree item ID for the specified display node
+	 *
+	 * @param DisplayNode The node to find a curve editor tree item for
+	 * @return The ID of the tree item for this display node in the curve editor, or FCurveEditorTreeItemID::Invalid() if one was not found.
+	 */
+	FCurveEditorTreeItemID FindCurveEditorTreeItem(TSharedRef<FSequencerDisplayNode> DisplayNode) const;
+
 public:
 
 	/**
diff --git a/Engine/Source/Editor/Sequencer/Private/SequencerSettings.cpp b/Engine/Source/Editor/Sequencer/Private/SequencerSettings.cpp
--- a/Engine/Source/Editor/Sequencer/Private/SequencerSettings.cpp
+++ b/Engine/Source/Editor/Sequencer/Private/SequencerSettings.cpp
@@ -33,6 +33,8 @@ USequencerSettings::USequencerSettings( const FObjectInitializer& ObjectInitiali
 	ZoomPosition = ESequencerZoomPosition::SZP_CurrentTime;
 	bAutoScrollEnabled = false;
 	bLinkCurveEditorTimeRange = false;
+	bSynchronizeCurveEditorSelection = true;
+	bIsolateCurveEditorToSelection = true;
 	LoopMode = ESequencerLoopMode::SLM_NoLoop;
 	bKeepCursorInPlayRangeWhileScrubbing = false;
 	bKeepCursorInPlayRange = true;
@@ -468,6 +470,23 @@ void USequencerSettings::SetLinkCurveEditorTimeRange(bool InbLinkCurveEditorTime
 	}
 }
 
+void USequencerSettings::SyncCurveEditorSelection(bool bInSynchronizeCurveEditorSelection)
+{
+	if (bSynchronizeCurveEditorSelection != bInSynchronizeCurveEditorSelection)
+	{
+		bSynchronizeCurveEditorSelection = bInSynchronizeCurveEditorSelection;
+		SaveConfig();
+	}
+}
+
+void USequencerSettings::IsolateCurveEditorToSelection(bool bInIsolateCurveEditorToSelection)
+{
+	if (bIsolateCurveEditorToSelection != bInIsolateCurveEditorToSelection)
+	{
+		bIsolateCurveEditorToSelection = bInIsolateCurveEditorToSelection;
+		SaveConfig();
+	}
+}
 
 uint8 USequencerSettings::GetZeroPadFrames() const
 {
diff --git a/Engine/Source/Editor/Sequencer/Public/ISequencerModule.h b/Engine/Source/Editor/Sequencer/Public/ISequencerModule.h
--- a/Engine/Source/Editor/Sequencer/Public/ISequencerModule.h
+++ b/Engine/Source/Editor/Sequencer/Public/ISequencerModule.h
@@ -21,6 +21,8 @@ class IToolkitHost;
 class UMovieSceneSequence;
 struct FSequencerInitParams;
 
+enum class ECurveEditorTreeFilterType : uint32;
+
 /** Forward declaration for the default templated channel interface. Include SequencerChannelInterface.h for full definition. */
 template<typename> struct TSequencerChannelInterface;
 
@@ -257,6 +259,12 @@ public:
 	 */
 	ISequencerChannelInterface* FindChannelEditorInterface(FName ChannelTypeName) const;
 
+
+	/**
+	 * Retrieve the unique identifer for the sequencer selection curve editor filter
+	 */
+	static ECurveEditorTreeFilterType GetSequencerSelectionFilterType();
+
 public:
 
 	/**
diff --git a/Engine/Source/Editor/Sequencer/Public/SequencerSettings.h b/Engine/Source/Editor/Sequencer/Public/SequencerSettings.h
--- a/Engine/Source/Editor/Sequencer/Public/SequencerSettings.h
+++ b/Engine/Source/Editor/Sequencer/Public/SequencerSettings.h
@@ -224,6 +224,16 @@ public:
 	/** Sets whether or not to link the curve editor time range. */
 	void SetLinkCurveEditorTimeRange(bool InbLinkCurveEditorTimeRange);
 
+	/** Return true if we are to synchronize the curve editor and sequencer trees */
+	bool ShouldSyncCurveEditorSelection() const { return bSynchronizeCurveEditorSelection; }
+	/** Assign whether we are to synchronize the curve editor and sequencer trees */
+	void SyncCurveEditorSelection(bool bInSynchronizeCurveEditorSelection);
+
+	/** Return true if we should filter the curve editor tree to only nodes that are relevant to the current sequencer selection */
+	bool ShouldIsolateToCurveEditorSelection() const { return bIsolateCurveEditorToSelection; }
+	/** Assign whether we should filter the curve editor tree to only nodes that are relevant to the current sequencer selection */
+	void IsolateCurveEditorToSelection(bool bInIsolateCurveEditorToSelection);
+
 	/** Gets the loop mode. */
 	ESequencerLoopMode GetLoopMode() const;
 	/** Sets the loop mode. */
@@ -422,6 +432,14 @@ protected:
 	UPROPERTY( config, EditAnywhere, Category=CurveEditor )
 	bool bLinkCurveEditorTimeRange;
 
+	/** When enabled, changing the sequencer tree selection will also select the relevant nodes in the curve editor tree if possible. */
+	UPROPERTY( config, EditAnywhere, Category=CurveEditor )
+	bool bSynchronizeCurveEditorSelection;
+
+	/** When enabled, changing the sequencer tree selection will isolate (auto-filter) the selected nodes in the curve editor. */
+	UPROPERTY( config, EditAnywhere, Category=CurveEditor )
+	bool bIsolateCurveEditorToSelection;
+
 	/** The loop mode of the playback in timeline. */
 	UPROPERTY( config )
 	TEnumAsByte<ESequencerLoopMode> LoopMode;
