Sequencer: Fixed a crash caused by stale tree item IDs being retained and selected

diff --git a/Engine/Source/Editor/CurveEditor/Private/Tree/CurveEditorTree.cpp b/Engine/Source/Editor/CurveEditor/Private/Tree/CurveEditorTree.cpp
--- a/Engine/Source/Editor/CurveEditor/Private/Tree/CurveEditorTree.cpp
+++ b/Engine/Source/Editor/CurveEditor/Private/Tree/CurveEditorTree.cpp
@@ -137,6 +137,11 @@ FCurveEditorTreeItem* FCurveEditorTree::FindItem(FCurveEditorTreeItemID ItemID)
 	return Items.Find(ItemID);
 }
 
+const FCurveEditorTreeItem* FCurveEditorTree::FindItem(FCurveEditorTreeItemID ItemID) const
+{
+	return Items.Find(ItemID);
+}
+
 const TArray<FCurveEditorTreeItemID>& FCurveEditorTree::GetRootItems() const
 {
 	return RootItems.ChildIDs;
@@ -400,6 +405,12 @@ void FCurveEditorTree::SetDirectSelection(TArray<FCurveEditorTreeItemID>&& TreeI
 	{
 		FCurveEditorTreeItemID ItemID = TreeItems[Index];
 
+		const FCurveEditorTreeItem* TreeItem = Items.Find(ItemID);
+		if (!ensureAlwaysMsgf(TreeItem, TEXT("Selected tree item does not exist. This must have bee applied externally.")))
+		{
+			continue;
+		}
+
 		if (Index < LastDirectlySelected)
 		{
 			Selection.Add(ItemID, ECurveEditorTreeSelectionState::Explicit);
@@ -409,7 +420,7 @@ void FCurveEditorTree::SetDirectSelection(TArray<FCurveEditorTreeItemID>&& TreeI
 			Selection.Add(ItemID, ECurveEditorTreeSelectionState::ImplicitChild);
 		}
 
-		for (FCurveEditorTreeItemID ChildID : Items.FindChecked(ItemID).GetChildren())
+		for (FCurveEditorTreeItemID ChildID : TreeItem->GetChildren())
 		{
 			if (FilterStates.Get(ChildID) != ECurveEditorTreeFilterState::NoMatch)
 			{
diff --git a/Engine/Source/Editor/CurveEditor/Private/Tree/SCurveEditorTree.cpp b/Engine/Source/Editor/CurveEditor/Private/Tree/SCurveEditorTree.cpp
--- a/Engine/Source/Editor/CurveEditor/Private/Tree/SCurveEditorTree.cpp
+++ b/Engine/Source/Editor/CurveEditor/Private/Tree/SCurveEditorTree.cpp
@@ -124,13 +124,19 @@ void SCurveEditorTree::RefreshTree()
 	{
 		// Add any currently selected items' parents to the expanded items array.
 		// This ensures that items that were selected during a filter operation remain expanded and selected when finished
-		for (FCurveEditorTreeItemID SelectedItem : GetSelectedItems())
+		for (FCurveEditorTreeItemID SelectedItemID : GetSelectedItems())
 		{
-			FCurveEditorTreeItemID ItemToExpand = CurveEditorTree->GetItem(SelectedItem).GetParentID();
-			while (ItemToExpand != FCurveEditorTreeItemID::Invalid())
+			// Add the selected item's parent and any grandparents to the list
+			const FCurveEditorTreeItem* ParentItem = CurveEditorTree->FindItem(SelectedItemID);
+			if (ParentItem)
 			{
-				PreFilterExpandedItems.Add(ItemToExpand);
-				ItemToExpand = CurveEditorTree->GetItem(ItemToExpand).GetParentID();
+				ParentItem = CurveEditorTree->FindItem(ParentItem->GetParentID());
+			}
+
+			while (ParentItem)
+			{
+				PreFilterExpandedItems.Add(ParentItem->GetID());
+				ParentItem = CurveEditorTree->FindItem(ParentItem->GetParentID());
 			}
 		}
 
diff --git a/Engine/Source/Editor/CurveEditor/Public/Tree/CurveEditorTree.h b/Engine/Source/Editor/CurveEditor/Public/Tree/CurveEditorTree.h
--- a/Engine/Source/Editor/CurveEditor/Public/Tree/CurveEditorTree.h
+++ b/Engine/Source/Editor/CurveEditor/Public/Tree/CurveEditorTree.h
@@ -353,6 +353,11 @@ public:
 	 */
 	FCurveEditorTreeItem* FindItem(FCurveEditorTreeItemID ItemID);
 
+	/**
+	 * Retrieve an item from its ID or nullptr if the ID is not valid
+	 */
+	const FCurveEditorTreeItem* FindItem(FCurveEditorTreeItemID ItemID) const;
+
 	/** 
 	 * Retrieve this curve editor's root items irrespective of filter state
 	 */
diff --git a/Engine/Source/Editor/Sequencer/Private/SequencerNodeTree.cpp b/Engine/Source/Editor/Sequencer/Private/SequencerNodeTree.cpp
--- a/Engine/Source/Editor/Sequencer/Private/SequencerNodeTree.cpp
+++ b/Engine/Source/Editor/Sequencer/Private/SequencerNodeTree.cpp
@@ -897,7 +897,10 @@ FCurveEditorTreeItemID FSequencerNodeTree::AddToCurveEditor(TSharedRef<FSequence
 {
 	if (FCurveEditorTreeItemID* Existing = CurveEditorTreeItemIDs.Find(DisplayNode))
 	{
-		return *Existing;
+		if (CurveEditor->GetTree()->FindItem(*Existing) != nullptr)
+		{
+			return *Existing;
+		}
 	}
 
 	TSharedPtr<FSequencerDisplayNode> Parent = DisplayNode->GetParent();
@@ -914,5 +917,9 @@ FCurveEditorTreeItemID FSequencerNodeTree::AddToCurveEditor(TSharedRef<FSequence
 FCurveEditorTreeItemID FSequencerNodeTree::FindCurveEditorTreeItem(TSharedRef<FSequencerDisplayNode> DisplayNode) const
 {
 	const FCurveEditorTreeItemID* FoundID = CurveEditorTreeItemIDs.Find(DisplayNode);
-	return FoundID ? *FoundID : FCurveEditorTreeItemID::Invalid();
+	if (FoundID && Sequencer.GetCurveEditor()->GetTree()->FindItem(*FoundID) != nullptr)
+	{
+		return *FoundID;
+	}
+	return FCurveEditorTreeItemID::Invalid();
 }
\ No newline at end of file
