Curve Editor: Guard against viewwidth 0.

diff --git a/Engine/Source/Editor/CurveEditor/Private/SCurveEditorView.cpp b/Engine/Source/Editor/CurveEditor/Private/SCurveEditorView.cpp
--- a/Engine/Source/Editor/CurveEditor/Private/SCurveEditorView.cpp
+++ b/Engine/Source/Editor/CurveEditor/Private/SCurveEditorView.cpp
@@ -37,14 +37,17 @@ void SCurveEditorView::GetInputBounds(double& OutInputMin, double& OutInputMax)
 		const float ContainerWidth = ContainerGeometry.GetLocalSize().X;
 		const float ViewWidth      = ViewGeometry.GetLocalSize().X;
 
-		const float LeftPixelCrop  = ViewGeometry.LocalToAbsolute(FVector2D(0.f, 0.f)).X                 - ContainerGeometry.LocalToAbsolute(FVector2D(0.f, 0.f)).X;
-		const float RightPixelCrop = ContainerGeometry.LocalToAbsolute(FVector2D(ContainerWidth, 0.f)).X - ViewGeometry.LocalToAbsolute(FVector2D(ViewWidth, 0.f)).X;
+		if (ViewWidth > 0.f)
+		{
+			const float LeftPixelCrop = ViewGeometry.LocalToAbsolute(FVector2D(0.f, 0.f)).X - ContainerGeometry.LocalToAbsolute(FVector2D(0.f, 0.f)).X;
+			const float RightPixelCrop = ContainerGeometry.LocalToAbsolute(FVector2D(ContainerWidth, 0.f)).X - ViewGeometry.LocalToAbsolute(FVector2D(ViewWidth, 0.f)).X;
 
-		const double ContainerInputPerPixel = (OutInputMax - OutInputMin) / ContainerWidth;
+			const double ContainerInputPerPixel = (OutInputMax - OutInputMin) / ContainerWidth;
 
-		// Offset by the total range first
-		OutInputMin += ContainerInputPerPixel * LeftPixelCrop;
-		OutInputMax -= ContainerInputPerPixel * RightPixelCrop;
+			// Offset by the total range first
+			OutInputMin += ContainerInputPerPixel * LeftPixelCrop;
+			OutInputMax -= ContainerInputPerPixel * RightPixelCrop;
+		}
 	}
 }
 
