Curve Editor: Mark as RF_Transactional

diff --git a/Engine/Source/Editor/CurveEditor/Private/RichCurveEditorModel.cpp b/Engine/Source/Editor/CurveEditor/Private/RichCurveEditorModel.cpp
--- a/Engine/Source/Editor/CurveEditor/Private/RichCurveEditorModel.cpp
+++ b/Engine/Source/Editor/CurveEditor/Private/RichCurveEditorModel.cpp
@@ -28,6 +28,7 @@ void FRichCurveEditorModel::Modify()
 {
 	if (UObject* Owner = WeakOwner.Get())
 	{
+		Owner->SetFlags(RF_Transactional);
 		Owner->Modify();
 	}
 }
