This is a significant overhaul to the Curve Editor used by Sequencer which adds a plugin-based architecture and extensibility. New tools and toolbar buttons can be added to all usages of the curve edi

diff --git a/Engine/Source/Editor/CurveEditor/Private/CurveEditor.cpp b/Engine/Source/Editor/CurveEditor/Private/CurveEditor.cpp
--- a/Engine/Source/Editor/CurveEditor/Private/CurveEditor.cpp
+++ b/Engine/Source/Editor/CurveEditor/Private/CurveEditor.cpp
@@ -2,7 +2,6 @@
 
 #include "CurveEditor.h"
 #include "Layout/Geometry.h"
-#include "CurveEditorScreenSpace.h"
 #include "CurveEditorSnapMetrics.h"
 #include "CurveEditorCommands.h"
 #include "CurveEditorSettings.h"
@@ -11,6 +10,17 @@
 #include "Framework/Commands/UICommandList.h"
 #include "Editor.h"
 #include "ScopedTransaction.h"
+#include "SCurveEditorView.h"
+#include "SCurveEditorPanel.h"
+#include "ICurveEditorExtension.h"
+#include "ICurveEditorToolExtension.h"
+#include "Tree/ICurveEditorTreeItem.h"
+#include "Filters/CurveEditorFilterBase.h"
+#include "Filters/SCurveEditorFilterPanel.h"
+#include "Editor/EditorEngine.h"
+#include "ITimeSlider.h"
+#include "Framework/Notifications/NotificationManager.h"
+#include "Widgets/Notifications/SNotificationList.h"
 
 #define LOCTEXT_NAMESPACE "CurveEditor"
 
@@ -25,9 +35,58 @@ FCurveModelID FCurveModelID::Unique()
 
 FCurveEditor::FCurveEditor()
 	: Bounds(new FStaticCurveEditorBounds)
-	, CachedPhysicalSize(0.f, 0.f)
+	, bBoundTransformUpdatesSuppressed(false)
+	, ActiveCurvesSerialNumber(0)
 {
 	Settings = GetMutableDefault<UCurveEditorSettings>();
+	OutputSnapEnabledAttribute = true;
+	InputSnapEnabledAttribute  = true;
+	OutputSnapIntervalAttribute = 0.1;
+	InputSnapRateAttribute = FFrameRate(10, 1);
+
+	GridLineLabelFormatXAttribute = LOCTEXT("GridXLabelFormat", "{0}s");
+	GridLineLabelFormatYAttribute = LOCTEXT("GridYLabelFormat", "{0}");
+}
+
+void FCurveEditor::InitCurveEditor(const FCurveEditorInitParams& InInitParams)
+{
+	ICurveEditorModule& CurveEditorModule = FModuleManager::LoadModuleChecked<ICurveEditorModule>("CurveEditor");
+
+	// Editor Extensions can be registered in the Curve Editor module. To allow users to derive from FCurveEditor
+	// we have to manually reach out to the module and get a list of extensions to create an instance of them.
+	// If none of your extensions are showing up, it's because you forgot to call this function after construction
+	// We're not allowed to use SharedThis(...) in a Constructor so it must exist as a separate function call.
+	TArrayView<const FOnCreateCurveEditorExtension> Extensions = CurveEditorModule.GetEditorExtensions();
+	for (int32 DelegateIndex = 0; DelegateIndex < Extensions.Num(); ++DelegateIndex)
+	{
+		check(Extensions[DelegateIndex].IsBound());
+
+		// We call a delegate and have the delegate create the instance to cover cross-module
+		TSharedRef<ICurveEditorExtension> NewExtension = Extensions[DelegateIndex].Execute(SharedThis(this));
+		EditorExtensions.Add(NewExtension);
+	}
+
+	TArrayView<const FOnCreateCurveEditorToolExtension> Tools = CurveEditorModule.GetToolExtensions();
+	for (int32 DelegateIndex = 0; DelegateIndex < Tools.Num(); ++DelegateIndex)
+	{
+		check(Tools[DelegateIndex].IsBound());
+
+		// We call a delegate and have the delegate create the instance to cover cross-module
+		AddTool(Tools[DelegateIndex].Execute(SharedThis(this)));
+	}
+
+	// Listen to global undo so we can fix up our selection state for keys that no longer exist.
+	GEditor->RegisterForUndo(this);
+}
+
+void FCurveEditor::SetPanel(TSharedPtr<SCurveEditorPanel> InPanel)
+{
+	WeakPanel = InPanel;
+}
+
+TSharedPtr<SCurveEditorPanel> FCurveEditor::GetPanel() const
+{
+	return WeakPanel.Pin();
 }
 
 FCurveModel* FCurveEditor::FindCurve(FCurveModelID CurveID) const
@@ -41,17 +100,150 @@ const TMap<FCurveModelID, TUniquePtr<FCurveModel>>& FCurveEditor::GetCurves() co
 	return CurveData;
 }
 
+FCurveEditorToolID FCurveEditor::AddTool(TUniquePtr<ICurveEditorToolExtension>&& InTool)
+{
+	FCurveEditorToolID NewID = FCurveEditorToolID::Unique();
+	ToolExtensions.Add(NewID, MoveTemp(InTool));
+	ToolExtensions[NewID]->SetToolID(NewID);
+	return NewID;
+}
+
 FCurveModelID FCurveEditor::AddCurve(TUniquePtr<FCurveModel>&& InCurve)
 {
 	FCurveModelID NewID = FCurveModelID::Unique();
 	CurveData.Add(NewID, MoveTemp(InCurve));
+
+	++ActiveCurvesSerialNumber;
+
+	return NewID;
+}
+
+FCurveModelID FCurveEditor::AddCurveForTreeItem(TUniquePtr<FCurveModel>&& InCurve, FCurveEditorTreeItemID TreeItemID)
+{
+	FCurveModelID NewID = FCurveModelID::Unique();
+
+	CurveData.Add(NewID, MoveTemp(InCurve));
+	TreeIDByCurveID.Add(NewID, TreeItemID);
+
+	++ActiveCurvesSerialNumber;
+
 	return NewID;
 }
 
 void FCurveEditor::RemoveCurve(FCurveModelID InCurveID)
 {
+	TSharedPtr<SCurveEditorPanel> Panel = WeakPanel.Pin();
+	if (Panel.IsValid())
+	{
+		Panel->RemoveCurveFromViews(InCurveID);
+	}
+
 	CurveData.Remove(InCurveID);
 	Selection.Remove(InCurveID);
+	PinnedCurves.Remove(InCurveID);
+
+	++ActiveCurvesSerialNumber;
+}
+
+bool FCurveEditor::IsCurvePinned(FCurveModelID InCurveID) const
+{
+	return PinnedCurves.Contains(InCurveID);
+}
+
+void FCurveEditor::PinCurve(FCurveModelID InCurveID)
+{
+	PinnedCurves.Add(InCurveID);
+	++ActiveCurvesSerialNumber;
+}
+
+void FCurveEditor::UnpinCurve(FCurveModelID InCurveID)
+{
+	PinnedCurves.Remove(InCurveID);
+	++ActiveCurvesSerialNumber;
+}
+
+const SCurveEditorView* FCurveEditor::FindFirstInteractiveView(FCurveModelID InCurveID) const
+{
+	for (auto ViewIt = GetPanel()->FindViews(InCurveID); ViewIt; ++ViewIt)
+	{
+		if (ViewIt.Value()->IsInteractive())
+		{
+			return &ViewIt.Value().Get();
+		}
+	}
+	return nullptr;
+}
+
+FCurveEditorTreeItem& FCurveEditor::GetTreeItem(FCurveEditorTreeItemID ItemID)
+{
+	return Tree.GetItem(ItemID);
+}
+
+const FCurveEditorTreeItem& FCurveEditor::GetTreeItem(FCurveEditorTreeItemID ItemID) const
+{
+	return Tree.GetItem(ItemID);
+}
+
+const TArray<FCurveEditorTreeItemID>& FCurveEditor::GetRootTreeItems() const
+{
+	return Tree.GetRootItems();
+}
+
+FCurveEditorTreeItem* FCurveEditor::AddTreeItem(FCurveEditorTreeItemID ParentID)
+{
+	return Tree.AddItem(ParentID);
+}
+
+void FCurveEditor::RemoveTreeItem(FCurveEditorTreeItemID ItemID)
+{
+	FCurveEditorTreeItem* Item = Tree.FindItem(ItemID);
+	if (!Item)
+	{
+		return;
+	}
+
+	Tree.RemoveItem(ItemID, this);
+	++ActiveCurvesSerialNumber;
+}
+
+FSimpleMulticastDelegate& FCurveEditor::OnTreeChanged()
+{
+	return Tree.OnChanged();
+}
+
+ECurveEditorTreeSelectionState FCurveEditor::GetTreeSelectionState(FCurveEditorTreeItemID InTreeItemID) const
+{
+	return Tree.GetSelectionState(InTreeItemID);
+}
+
+void FCurveEditor::SetDirectTreeSelection(TArray<FCurveEditorTreeItemID>&& TreeItems)
+{
+	TMap<FCurveEditorTreeItemID, ECurveEditorTreeSelectionState> PreviousSelection = Tree.GetSelection();
+	Tree.SetDirectSelection(MoveTemp(TreeItems));
+
+	const TMap<FCurveEditorTreeItemID, ECurveEditorTreeSelectionState>& NewSelection = Tree.GetSelection();
+	for (TTuple<FCurveEditorTreeItemID, ECurveEditorTreeSelectionState> CachedItem : PreviousSelection)
+	{
+		const ECurveEditorTreeSelectionState* NewState = NewSelection.Find(CachedItem.Key);
+		if (!NewState || *NewState == ECurveEditorTreeSelectionState::None)
+		{
+			GetTreeItem(CachedItem.Key).DestroyUnpinnedCurves(this);
+		}
+	}
+
+	// Ensure the new selection has valid curve models
+	for (TTuple<FCurveEditorTreeItemID, ECurveEditorTreeSelectionState> CachedItem : Tree.GetSelection())
+	{
+		if (CachedItem.Value != ECurveEditorTreeSelectionState::None)
+		{
+			GetTreeItem(CachedItem.Key).GetOrCreateCurves(this);
+		}
+	}
+}
+
+const TMap<FCurveEditorTreeItemID, ECurveEditorTreeSelectionState>& FCurveEditor::GetTreeSelection() const
+{
+	return Tree.GetSelection();
 }
 
 void FCurveEditor::SetBounds(TUniquePtr<ICurveEditorBounds>&& InBounds)
@@ -75,10 +267,14 @@ void FCurveEditor::BindCommands()
 	CommandList->MapAction(FGenericCommands::Get().Redo,   FExecuteAction::CreateLambda([]{ GEditor->RedoTransaction(); }));
 	CommandList->MapAction(FGenericCommands::Get().Delete, FExecuteAction::CreateSP(this, &FCurveEditor::DeleteSelection));
 
-	CommandList->MapAction(FCurveEditorCommands::Get().ZoomToFitAll,        FExecuteAction::CreateSP(this, &FCurveEditor::ZoomToFit, EAxisList::All));
-	CommandList->MapAction(FCurveEditorCommands::Get().ZoomToFit,           FExecuteAction::CreateSP(this, &FCurveEditor::ZoomToFitSelection, EAxisList::All));
-	CommandList->MapAction(FCurveEditorCommands::Get().ZoomToFitHorizontal, FExecuteAction::CreateSP(this, &FCurveEditor::ZoomToFitSelection, EAxisList::X));
-	CommandList->MapAction(FCurveEditorCommands::Get().ZoomToFitVertical,   FExecuteAction::CreateSP(this, &FCurveEditor::ZoomToFitSelection, EAxisList::Y));
+	CommandList->MapAction(FCurveEditorCommands::Get().ZoomToFit, FExecuteAction::CreateSP(this, &FCurveEditor::ZoomToFit, EAxisList::All));
+
+	CommandList->MapAction(FCurveEditorCommands::Get().StepToNextKey, FExecuteAction::CreateSP(this, &FCurveEditor::StepToNextKey));
+	CommandList->MapAction(FCurveEditorCommands::Get().StepToPreviousKey, FExecuteAction::CreateSP(this, &FCurveEditor::StepToPreviousKey));
+	CommandList->MapAction(FCurveEditorCommands::Get().StepForward, FExecuteAction::CreateSP(this, &FCurveEditor::StepForward), EUIActionRepeatMode::RepeatEnabled);
+	CommandList->MapAction(FCurveEditorCommands::Get().StepBackward, FExecuteAction::CreateSP(this, &FCurveEditor::StepBackward), EUIActionRepeatMode::RepeatEnabled);
+	CommandList->MapAction(FCurveEditorCommands::Get().JumpToStart, FExecuteAction::CreateSP(this, &FCurveEditor::JumpToStart));
+	CommandList->MapAction(FCurveEditorCommands::Get().JumpToEnd, FExecuteAction::CreateSP(this, &FCurveEditor::JumpToEnd));
 
 	{
 		FExecuteAction   ToggleInputSnapping     = FExecuteAction::CreateSP(this,   &FCurveEditor::ToggleInputSnapping);
@@ -90,10 +286,12 @@ void FCurveEditor::BindCommands()
 		CommandList->MapAction(FCurveEditorCommands::Get().ToggleOutputSnapping, ToggleOutputSnapping, FCanExecuteAction(), IsOutputSnappingEnabled);
 	}
 
-	CommandList->MapAction(FCurveEditorCommands::Get().FlattenTangents,    FExecuteAction::CreateSP(this, &FCurveEditor::FlattenSelection));
-	CommandList->MapAction(FCurveEditorCommands::Get().StraightenTangents, FExecuteAction::CreateSP(this, &FCurveEditor::StraightenSelection));
-	CommandList->MapAction(FCurveEditorCommands::Get().BakeCurve,          FExecuteAction::CreateSP(this, &FCurveEditor::BakeSelection));
-	CommandList->MapAction(FCurveEditorCommands::Get().ReduceCurve,        FExecuteAction::CreateSP(this, &FCurveEditor::SimplifySelection, 0.1f));
+	// Flatten and Straighten Tangents
+	{
+		CommandList->MapAction(FCurveEditorCommands::Get().FlattenTangents, FExecuteAction::CreateSP(this, &FCurveEditor::FlattenSelection), FCanExecuteAction::CreateSP(this, &FCurveEditor::CanFlattenOrStraightenSelection) );
+		CommandList->MapAction(FCurveEditorCommands::Get().StraightenTangents, FExecuteAction::CreateSP(this, &FCurveEditor::StraightenSelection), FCanExecuteAction::CreateSP(this, &FCurveEditor::CanFlattenOrStraightenSelection) );
+	}
+
 
 	// Tangent Visibility
 	{
@@ -120,18 +318,22 @@ void FCurveEditor::BindCommands()
 		FExecuteAction::CreateLambda( [CurveSettings]{ CurveSettings->SetShowCurveEditorCurveToolTips( !CurveSettings->GetShowCurveEditorCurveToolTips() ); } ),
 		FCanExecuteAction(),
 		FIsActionChecked::CreateLambda( [CurveSettings]{ return CurveSettings->GetShowCurveEditorCurveToolTips(); } ) );
-}
 
-FCurveEditorScreenSpace FCurveEditor::GetScreenSpace() const
-{
-	// Retrieve the current bounds and cache them
-	double InputMin = 0, InputMax = 1;
-	Bounds->GetInputBounds(InputMin, InputMax);
-
-	double OutputMin = 0, OutputMax = 1;
-	Bounds->GetOutputBounds(OutputMin, OutputMax);
+	// Deactivate Current Tool
+	CommandList->MapAction(FCurveEditorCommands::Get().DeactivateCurrentTool,
+		FExecuteAction::CreateSP(this, &FCurveEditor::MakeToolActive, FCurveEditorToolID::Unset()));
+		
+	// Bind commands for Editor Extensions
+	for (TSharedRef<ICurveEditorExtension> Extension : EditorExtensions)
+	{
+		Extension->BindCommands(CommandList.ToSharedRef());
+	}
 
-	return FCurveEditorScreenSpace(CachedPhysicalSize, InputMin, InputMax, OutputMin, OutputMax);
+	// Bind Commands for Tool Extensions
+	for (TPair<FCurveEditorToolID, TUniquePtr<ICurveEditorToolExtension>>& Pair : ToolExtensions)
+	{
+		Pair.Value->BindCommands(CommandList.ToSharedRef());
+	}
 }
 
 FCurveEditorSnapMetrics FCurveEditor::GetSnapMetrics() const
@@ -146,175 +348,301 @@ FCurveEditorSnapMetrics FCurveEditor::GetSnapMetrics() const
 	return Metrics;
 }
 
-void FCurveEditor::UpdateGeometry(const FGeometry& NewGeometry)
+void FCurveEditor::ZoomToFit(EAxisList::Type Axes)
+{
+	// If they have keys selected, we fit the specific keys.
+	if (Selection.Count() > 0)
+	{
+		ZoomToFitSelection(Axes);
+	}
+	else
+	{
+		TMap<FCurveModelID, FKeyHandleSet> AllCurves;
+		for (FCurveModelID ID : GetEditedCurves())
+		{
+			AllCurves.Add(ID);
+		}
+		ZoomToFitInternal(Axes, AllCurves);
+	}
+}
+
+void FCurveEditor::ZoomToFitCurves(TArrayView<const FCurveModelID> CurveModelIDs, EAxisList::Type Axes)
 {
-	if (CachedPhysicalSize.X != 0 && CachedPhysicalSize.Y != 0)
+	TMap<FCurveModelID, FKeyHandleSet> AllCurves;
+	for (FCurveModelID ID : CurveModelIDs)
 	{
-		// Retrieve the current bounds and cache them
-		double InputMin = 0, InputMax = 1;
-		Bounds->GetInputBounds(InputMin, InputMax);
+		AllCurves.Add(ID);
+	}
+	ZoomToFitInternal(Axes, AllCurves);
+}
 
-		double OutputMin = 0, OutputMax = 1;
-		Bounds->GetOutputBounds(OutputMin, OutputMax);
+void FCurveEditor::ZoomToFitSelection(EAxisList::Type Axes)
+{
+	ZoomToFitInternal(Axes, Selection.GetAll());
+}
 
-		// Increase the visible input/output ranges based on the new size of the panel
-		const float OldWidth = CachedPhysicalSize.X;
-		if (NewGeometry.GetLocalSize().X != OldWidth)
-		{
-			const double PixelToInputRatio = (InputMax - InputMin) / OldWidth;
-			InputMax += PixelToInputRatio * (NewGeometry.GetLocalSize().X - OldWidth);
+void FCurveEditor::ZoomToFitInternal(EAxisList::Type Axes, const TMap<FCurveModelID, FKeyHandleSet>& CurveKeySet)
+{
+	TArray<FKeyPosition> KeyPositionsScratch;
 
-			Bounds->SetInputBounds(InputMin, InputMax);
+	double InputMin = TNumericLimits<double>::Max(), InputMax = TNumericLimits<double>::Lowest();
+
+	TMap<TSharedRef<SCurveEditorView>, TTuple<double, double>> ViewToOutputBounds;
+
+	for (const TTuple<FCurveModelID, FKeyHandleSet>& Pair : CurveKeySet)
+	{
+		FCurveModelID CurveID = Pair.Key;
+		const FCurveModel* Curve = FindCurve(CurveID);
+		if (!Curve)
+		{
+			continue;
 		}
 
-		const float OldHeight = CachedPhysicalSize.Y;
-		if (NewGeometry.GetLocalSize().Y != OldHeight)
+		double OutputMin = TNumericLimits<double>::Max(), OutputMax = TNumericLimits<double>::Lowest();
+
+		int32 NumKeys = Pair.Value.AsArray().Num();
+		if (NumKeys == 0)
+		{
+			double LocalMin = 0.0, LocalMax = 1.0;
+
+			// Zoom to the entire curve range if no specific keys are specified
+			if (Curve->GetNumKeys())
+			{
+				// Only zoom time range if there are keys on the curve (otherwise where do we zoom *to* on an infinite timeline?)
+				Curve->GetTimeRange(LocalMin, LocalMax);
+				InputMin = FMath::Min(InputMin, LocalMin);
+				InputMax = FMath::Max(InputMax, LocalMax);
+			}
+
+			// Most curve types we know about support default values, so we can zoom to that even if there are no keys
+			Curve->GetValueRange(LocalMin, LocalMax);
+			OutputMin = FMath::Min(OutputMin, LocalMin);
+			OutputMax = FMath::Max(OutputMax, LocalMax);
+		}
+		else
 		{
-			const double PixelToOutputRatio = (OutputMax - OutputMin) / OldHeight;
-			OutputMin -= PixelToOutputRatio * (NewGeometry.GetLocalSize().Y - OldHeight);
+			// Zoom to the min/max of the specified key set
+			KeyPositionsScratch.SetNum(NumKeys, false);
+			Curve->GetKeyPositions(Pair.Value.AsArray(), KeyPositionsScratch);
+			for (const FKeyPosition& Key : KeyPositionsScratch)
+			{
+				InputMin  = FMath::Min(InputMin, Key.InputValue);
+				InputMax  = FMath::Max(InputMax, Key.InputValue);
+				OutputMin = FMath::Min(OutputMin, Key.OutputValue);
+				OutputMax = FMath::Max(OutputMax, Key.OutputValue);
+			}
+		}
 
-			Bounds->SetOutputBounds(OutputMin, OutputMax);
+		if (Axes & EAxisList::Y)
+		{
+			// Store the min max for each view
+			for (auto ViewIt = GetPanel()->FindViews(CurveID); ViewIt; ++ViewIt)
+			{
+				TTuple<double, double>* ViewBounds = ViewToOutputBounds.Find(ViewIt.Value());
+				if (ViewBounds)
+				{
+					ViewBounds->Get<0>() = FMath::Min(ViewBounds->Get<0>(), OutputMin);
+					ViewBounds->Get<1>() = FMath::Max(ViewBounds->Get<1>(), OutputMax);
+				}
+				else
+				{
+					ViewToOutputBounds.Add(ViewIt.Value(), MakeTuple(OutputMin, OutputMax));
+				}
+			}
 		}
 	}
 
-	CachedPhysicalSize = NewGeometry.GetLocalSize();
-}
+	FCurveEditorSnapMetrics SnapMetrics = GetSnapMetrics();
 
-void FCurveEditor::Zoom(float Amount)
-{
-	FCurveEditorScreenSpace ScreenSpace = GetScreenSpace();
+	if (Axes & EAxisList::X && InputMin != TNumericLimits<double>::Max() && InputMax != TNumericLimits<double>::Lowest())
+	{
+		// If zooming to the same (or invalid) min/max, keep the same zoom scale and center within the timeline
+		if (InputMin >= InputMax)
+		{
+			double CurrentInputMin = 0.0, CurrentInputMax = 1.0;
+			Bounds->GetInputBounds(CurrentInputMin, CurrentInputMax);
 
-	const double InputOrigin  = (ScreenSpace.GetInputMax()  - ScreenSpace.GetInputMin())  * 0.5;
-	const double OutputOrigin = (ScreenSpace.GetOutputMax() - ScreenSpace.GetOutputMin()) * 0.5;
+			const double HalfInputScale = (CurrentInputMax - CurrentInputMin)*0.5;
+			InputMin -= HalfInputScale;
+			InputMax += HalfInputScale;
+		}
+		else
+		{
+			const double MinInputZoom = SnapMetrics.bSnapInputValues ? SnapMetrics.InputSnapRate.AsInterval() : 0.00001;
+			const double InputPadding = FMath::Max((InputMax - InputMin) * 0.1, MinInputZoom);
+			InputMax = FMath::Max(InputMin + MinInputZoom, InputMax);
 
-	ZoomAround(Amount, InputOrigin, OutputOrigin);
-}
+			InputMin -= InputPadding;
+			InputMax += InputPadding;
+		}
 
-void FCurveEditor::ZoomAround(float Amount, double InputOrigin, double OutputOrigin)
-{
-	FCurveEditorScreenSpace ScreenSpace = GetScreenSpace();
+		Bounds->SetInputBounds(InputMin, InputMax);
+	}
 
-	double InputMin  = ScreenSpace.GetInputMin();
-	double InputMax  = ScreenSpace.GetInputMax();
-	double OutputMin = ScreenSpace.GetOutputMin();
-	double OutputMax = ScreenSpace.GetOutputMax();
+	// Perform per-view output zoom for any computed ranges
+	for (const TTuple<TSharedRef<SCurveEditorView>, TTuple<double, double>>& ViewAndBounds : ViewToOutputBounds)
+	{
+		TSharedRef<SCurveEditorView> View = ViewAndBounds.Key;
+		double OutputMin = ViewAndBounds.Value.Get<0>();
+		double OutputMax = ViewAndBounds.Value.Get<1>();
 
-	InputMin = InputOrigin - (InputOrigin - InputMin) * Amount;
-	InputMax = InputOrigin + (InputMax - InputOrigin) * Amount;
+		// If zooming to the same (or invalid) min/max, keep the same zoom scale and center within the timeline
+		if (OutputMin >= OutputMax)
+		{
+			const double HalfOutputScale = (View->GetOutputMax() - View->GetOutputMin())*0.5;
+			OutputMin -= HalfOutputScale;
+			OutputMax += HalfOutputScale;
+		}
+		else
+		{
+			const double MinOutputZoom = SnapMetrics.bSnapOutputValues ? SnapMetrics.OutputSnapInterval : 0.00001;
+			const double OutputPadding = FMath::Max((OutputMax - OutputMin) * 0.05, MinOutputZoom);
 
-	OutputMin = OutputOrigin - (OutputOrigin - OutputMin) * Amount;
-	OutputMax = OutputOrigin + (OutputMax - OutputOrigin) * Amount;
+			
+			OutputMin -= OutputPadding;
+			OutputMax = FMath::Max(OutputMin + MinOutputZoom, OutputMax) + OutputPadding;
+		}
 
-	Bounds->SetInputBounds(InputMin, InputMax);
-	Bounds->SetOutputBounds(OutputMin, OutputMax);
+		View->SetOutputBounds(OutputMin, OutputMax);
+	}
 }
 
-void FCurveEditor::ZoomToFit(EAxisList::Type Axes)
+void FCurveEditor::StepToNextKey()
 {
-	TArray<FCurveModelID> CurveModelIDs;
-	for (const TTuple<FCurveModelID, TUniquePtr<FCurveModel>>& Pair : CurveData)
+	if (!WeakTimeSliderController.IsValid())
 	{
-		CurveModelIDs.Add(Pair.Key);
+		return;
 	}
 
-	ZoomToFitCurves(MakeArrayView(CurveModelIDs.GetData(), CurveModelIDs.Num()), Axes);
-}
+	FFrameRate TickResolution = WeakTimeSliderController.Pin()->GetTickResolution();
 
-void FCurveEditor::ZoomToFitCurves(TArrayView<const FCurveModelID> CurveModelIDs, EAxisList::Type Axes)
-{
-	if (CurveData.Num() && CurveModelIDs.Num())
-	{
-		TArrayView<const FKeyHandle> Empty;
+	double CurrentTime = WeakTimeSliderController.Pin()->GetScrubPosition() / TickResolution;
 
-		double InputMin = TNumericLimits<double>::Max(), InputMax = TNumericLimits<double>::Lowest(), OutputMin = TNumericLimits<double>::Max(), OutputMax = TNumericLimits<double>::Lowest();
+	TOptional<double> NextTime;
 
-		TArray<FKeyHandle> KeyHandlesScratch;
-		TArray<FKeyPosition> KeyPositionsScratch;
+	for (const TTuple<FCurveModelID, TUniquePtr<FCurveModel>>& Pair : CurveData)
+	{
+		FCurveModel* CurveModel = Pair.Value.Get();
 
-		for (const FCurveModelID CurveModelID : CurveModelIDs)
+		if (CurveModel)
 		{
-			if (const FCurveModel* Curve = FindCurve(CurveModelID))
+			TArray<FKeyHandle> KeyHandles;
+			double MaxTime = NextTime.IsSet() ? NextTime.GetValue() : TNumericLimits<double>::Max();
+			CurveModel->GetKeys(*this, CurrentTime, MaxTime, TNumericLimits<double>::Lowest(), TNumericLimits<double>::Max(), KeyHandles);
+
+			TArray<FKeyPosition> KeyPositions;
+			KeyPositions.SetNum(KeyHandles.Num());
+			CurveModel->GetKeyPositions(TArrayView<FKeyHandle>(KeyHandles), KeyPositions);
+
+			for (FKeyPosition KeyPosition : KeyPositions)
 			{
-				double LocalMin, LocalMax;
-				Curve->GetTimeRange(LocalMin, LocalMax);
-				InputMin = FMath::Min(InputMin, LocalMin);
-				InputMax = FMath::Max(InputMax, LocalMax);
-				Curve->GetValueRange(LocalMin, LocalMax);
-				OutputMin = FMath::Min(OutputMin, LocalMin);
-				OutputMax = FMath::Max(OutputMax, LocalMax);
+				if (KeyPosition.InputValue > CurrentTime)
+				{
+					if (!NextTime.IsSet() || KeyPosition.InputValue < NextTime.GetValue())
+					{
+						NextTime = KeyPosition.InputValue;
+					}
+				}
 			}
 		}
+	}
 
-		if (InputMin != TNumericLimits<double>::Max() && InputMax != TNumericLimits<double>::Lowest() && 
-			OutputMin != TNumericLimits<double>::Max() && OutputMax != TNumericLimits<double>::Lowest())
-		{
-			ZoomToFitInternal(Axes, InputMin, InputMax, OutputMin, OutputMax);
-		}
+	if (NextTime.IsSet())
+	{
+		WeakTimeSliderController.Pin()->SetScrubPosition(NextTime.GetValue() * TickResolution);
 	}
 }
 
-void FCurveEditor::ZoomToFitSelection(EAxisList::Type Axes)
+void FCurveEditor::StepToPreviousKey()
 {
-	if (Selection.Count() <= 1)
+	if (!WeakTimeSliderController.IsValid())
 	{
-		ZoomToFit(Axes);
+		return;
 	}
-	else
-	{
-		TArray<FKeyPosition> KeyPositionsScratch;
 
-		double InputMin = TNumericLimits<double>::Max(), InputMax = TNumericLimits<double>::Lowest(), OutputMin = TNumericLimits<double>::Max(), OutputMax = TNumericLimits<double>::Lowest();
+	FFrameRate TickResolution = WeakTimeSliderController.Pin()->GetTickResolution();
+
+	double CurrentTime = WeakTimeSliderController.Pin()->GetScrubPosition() / TickResolution;
 
-		for (const TTuple<FCurveModelID, FKeyHandleSet>& Pair : Selection.GetAll())
+	TOptional<double> PreviousTime;
+
+	for (const TTuple<FCurveModelID, TUniquePtr<FCurveModel>>& Pair : CurveData)
+	{
+		FCurveModel* CurveModel = Pair.Value.Get();
+
+		if (CurveModel)
 		{
-			if (const FCurveModel* Curve = FindCurve(Pair.Key))
-			{
-				KeyPositionsScratch.SetNum(Pair.Value.AsArray().Num());
+			TArray<FKeyHandle> KeyHandles;
+			double MinTime = PreviousTime.IsSet() ? PreviousTime.GetValue() : TNumericLimits<double>::Lowest();
+			CurveModel->GetKeys(*this, MinTime, CurrentTime, TNumericLimits<double>::Lowest(), TNumericLimits<double>::Max(), KeyHandles);
 
-				Curve->GetKeyPositions(Pair.Value.AsArray(), KeyPositionsScratch);
-				for (FKeyPosition Key : KeyPositionsScratch)
+			TArray<FKeyPosition> KeyPositions;
+			KeyPositions.SetNum(KeyHandles.Num());
+			CurveModel->GetKeyPositions(TArrayView<FKeyHandle>(KeyHandles), KeyPositions);
+
+			for (FKeyPosition KeyPosition : KeyPositions)
+			{
+				if (KeyPosition.InputValue < CurrentTime)
 				{
-					InputMin  = FMath::Min(InputMin, Key.InputValue);
-					InputMax  = FMath::Max(InputMax, Key.InputValue);
-					OutputMin = FMath::Min(OutputMin, Key.OutputValue);
-					OutputMax = FMath::Max(OutputMax, Key.OutputValue);
+					if (!PreviousTime.IsSet() || KeyPosition.InputValue > PreviousTime.GetValue())
+					{
+						PreviousTime = KeyPosition.InputValue;
+					}
 				}
 			}
 		}
+	}
 
-		if (InputMin != TNumericLimits<double>::Max() && InputMax != TNumericLimits<double>::Lowest() && 
-			OutputMin != TNumericLimits<double>::Max() && OutputMax != TNumericLimits<double>::Lowest())
-		{
-			ZoomToFitInternal(Axes, InputMin, InputMax, OutputMin, OutputMax);
-		}
+	if (PreviousTime.IsSet())
+	{
+		WeakTimeSliderController.Pin()->SetScrubPosition(PreviousTime.GetValue() * TickResolution);
 	}
 }
 
-void FCurveEditor::ZoomToFitInternal(EAxisList::Type Axes, double InputMin, double InputMax, double OutputMin, double OutputMax)
+
+void FCurveEditor::StepForward()
 {
-	FCurveEditorSnapMetrics SnapMetrics = GetSnapMetrics();
-	double MinInputZoom = SnapMetrics.bSnapInputValues ? SnapMetrics.InputSnapRate.AsInterval() : 0.00001;
-	double MinOutputZoom = SnapMetrics.bSnapOutputValues ? SnapMetrics.OutputSnapInterval : 0.00001;
+	if (!WeakTimeSliderController.IsValid())
+	{
+		return;
+	}
 
-	InputMax = FMath::Max(InputMin + MinInputZoom, InputMax);
-	OutputMax = FMath::Max(OutputMin + MinOutputZoom, OutputMax);
+	FFrameRate TickResolution = WeakTimeSliderController.Pin()->GetTickResolution();
+	FFrameRate DisplayRate = WeakTimeSliderController.Pin()->GetDisplayRate();
 
-	double InputPadding  = (InputMax - InputMin) * 0.1;
-	double OutputPadding = (OutputMax - OutputMin) * 0.05;
+	FFrameTime OneFrame = FFrameRate::TransformTime(FFrameTime(1), DisplayRate, TickResolution);
+
+	WeakTimeSliderController.Pin()->SetScrubPosition(WeakTimeSliderController.Pin()->GetScrubPosition() + OneFrame);
+}
+
+void FCurveEditor::StepBackward()
+{
+	FFrameRate TickResolution = WeakTimeSliderController.Pin()->GetTickResolution();
+	FFrameRate DisplayRate = WeakTimeSliderController.Pin()->GetDisplayRate();
+
+	FFrameTime OneFrame = FFrameRate::TransformTime(FFrameTime(1), DisplayRate, TickResolution);
+
+	WeakTimeSliderController.Pin()->SetScrubPosition(WeakTimeSliderController.Pin()->GetScrubPosition() - OneFrame);
+}
 
-	InputMin -= InputPadding;
-	InputMax += InputPadding;
-	if (Axes & EAxisList::X)
+void FCurveEditor::JumpToStart()
+{
+	if (!WeakTimeSliderController.IsValid())
 	{
-		Bounds->SetInputBounds(InputMin - InputPadding, InputMax + InputPadding);
+		return;
 	}
 
-	OutputMin -= OutputPadding;
-	OutputMax += OutputPadding;
-	if (Axes & EAxisList::Y)
+	WeakTimeSliderController.Pin()->SetScrubPosition(WeakTimeSliderController.Pin()->GetPlayRange().GetLowerBoundValue());
+}
+
+void FCurveEditor::JumpToEnd()
+{
+	if (!WeakTimeSliderController.IsValid())
 	{
-		Bounds->SetOutputBounds(OutputMin - OutputPadding, OutputMax + OutputPadding);
+		return;
 	}
+
+	WeakTimeSliderController.Pin()->SetScrubPosition(WeakTimeSliderController.Pin()->GetPlayRange().GetUpperBoundValue());
 }
 
 bool FCurveEditor::IsInputSnappingEnabled() const
@@ -355,221 +683,43 @@ void FCurveEditor::ToggleOutputSnapping()
 	}
 }
 
-FVector2D FCurveEditor::GetVectorFromSlopeAndLength(float Slope, float Length)
+FCurveEditorScreenSpaceH FCurveEditor::GetPanelInputSpace() const
 {
-	float x = Length / FMath::Sqrt(Slope*Slope + 1.f);
-	float y = Slope * x;
-	return FVector2D(x, y);
-}
-
-FVector2D  FCurveEditor::GetTangentPositionInScreenSpace(const FVector2D &StartPos, float Tangent, float Weight) const
-{
-	FCurveEditorScreenSpace ScreenSpace = GetScreenSpace();
-	const float Angle = FMath::Atan(-Tangent);
-	float X, Y;
-	FMath::SinCos(&Y, &X, Angle);
-	X *= Weight;
-	Y *= Weight;
+	const float PanelWidth = FMath::Max(1.f, WeakPanel.Pin()->GetViewContainerGeometry().GetLocalSize().X);
 
-	X *= ScreenSpace.PixelsPerInput();
-	Y *= ScreenSpace.PixelsPerOutput();
-	return FVector2D(StartPos.X + X, StartPos.Y +Y);
-}
-
-void FCurveEditor::GetTangentAndWeightFromScreenPosition(const FVector2D &StartPos, const  FVector2D &TangentPos, float &Tangent, float &Weight) const
-{
-	FCurveEditorScreenSpace ScreenSpace = GetScreenSpace();
-	float X = ScreenSpace.ScreenToSeconds(TangentPos.X) - ScreenSpace.ScreenToSeconds(StartPos.X);
-	float Y = ScreenSpace.ScreenToValue(TangentPos.Y) - ScreenSpace.ScreenToValue(StartPos.Y);
-
-	Tangent = Y / X;
-	Weight = FMath::Sqrt(X*X + Y * Y);
-}
-
-void FCurveEditor::GetCurveDrawParams(TArray<FCurveDrawParams>& OutDrawParams) const
-{
-	FCurveEditorScreenSpace ScreenSpace = GetScreenSpace();
-	const float DisplayRatio = (ScreenSpace.PixelsPerOutput() / ScreenSpace.PixelsPerInput());
-
-	double InputMin = 0, InputMax = 1;
+	double InputMin = 0.0, InputMax = 1.0;
 	Bounds->GetInputBounds(InputMin, InputMax);
 
-	double OutputMin = 0, OutputMax = 1;
-	Bounds->GetOutputBounds(OutputMin, OutputMax);
-
-	ECurveEditorTangentVisibility::Type TangentVisibility = Settings->GetTangentVisibility();
-
-	OutDrawParams.Reserve(CurveData.Num());
-	for (const TTuple<FCurveModelID, TUniquePtr<FCurveModel>>& Pair : CurveData)
-	{
-		const FKeyHandleSet* SelectedKeys = Selection.GetAll().Find(Pair.Key);
-
-		FCurveDrawParams Params(Pair.Key);
-		Params.Color = Pair.Value->GetColor();
-
-		// Gather the display metrics to use for each key type
-		Pair.Value->GetKeyDrawInfo(ECurvePointType::Key, Params.KeyDrawInfo);
-		Pair.Value->GetKeyDrawInfo(ECurvePointType::ArriveTangent, Params.ArriveTangentDrawInfo);
-		Pair.Value->GetKeyDrawInfo(ECurvePointType::LeaveTangent, Params.LeaveTangentDrawInfo);
-
-		// Gather the interpolating points in input/output space
-		TArray<TTuple<double, double>> InterpolatingPoints;
-
-		Pair.Value->DrawCurve(*this, InterpolatingPoints);
-		Params.InterpolatingPoints.Reserve(InterpolatingPoints.Num());
-
-		double InputOffset = Pair.Value->GetInputDisplayOffset();
-
-		// Convert the interpolating points to screen space
-		for (TTuple<double, double> Point : InterpolatingPoints)
-		{
-			Params.InterpolatingPoints.Add(
-				FVector2D(
-					ScreenSpace.SecondsToScreen(Point.Get<0>() + InputOffset),
-					ScreenSpace.ValueToScreen(Point.Get<1>())
-				)
-			);
-		}
-
-		TArray<FKeyHandle> VisibleKeys;
-		Pair.Value->GetKeys(*this, InputMin, InputMax, OutputMin, OutputMax, VisibleKeys);
-
-		if (VisibleKeys.Num())
-		{
-			TArray<FKeyPosition> AllKeyPositions;
-			TArray<FKeyAttributes> AllKeyAttributes;
-
-			AllKeyPositions.SetNum(VisibleKeys.Num());
-			AllKeyAttributes.SetNum(VisibleKeys.Num());
-
-			Pair.Value->GetKeyPositions(VisibleKeys, AllKeyPositions);
-			Pair.Value->GetKeyAttributes(VisibleKeys, AllKeyAttributes);
-
-			for (int32 Index = 0; Index < VisibleKeys.Num(); ++Index)
-			{
-				const FKeyHandle      KeyHandle   = VisibleKeys[Index];
-				const FKeyPosition&   KeyPosition = AllKeyPositions[Index];
-				const FKeyAttributes& Attributes  = AllKeyAttributes[Index];
-
-				bool bShowTangents = TangentVisibility == ECurveEditorTangentVisibility::AllTangents || ( TangentVisibility == ECurveEditorTangentVisibility::SelectedKeys && SelectedKeys && SelectedKeys->Contains(VisibleKeys[Index]) );
-
-				// Add this key
-				FCurvePointInfo Key(KeyHandle);
-				Key.ScreenPosition = FVector2D(ScreenSpace.SecondsToScreen(KeyPosition.InputValue + InputOffset), ScreenSpace.ValueToScreen(KeyPosition.OutputValue));
-				Key.LayerBias = 2;
-
-				Params.Points.Add(Key);
-
-				if (bShowTangents && Attributes.HasArriveTangent())
-				{
-					float ArriveTangent = Attributes.GetArriveTangent();
-
-					FCurvePointInfo ArriveTangentPoint(KeyHandle);
-					ArriveTangentPoint.Type = ECurvePointType::ArriveTangent;
-
-
-					if (Attributes.HasTangentWeightMode() && Attributes.HasArriveTangentWeight() &&
-						(Attributes.GetTangentWeightMode() == RCTWM_WeightedBoth || Attributes.GetTangentWeightMode() == RCTWM_WeightedArrive))
-					{
-						ArriveTangentPoint.ScreenPosition = GetTangentPositionInScreenSpace(Key.ScreenPosition, ArriveTangent, -Attributes.GetArriveTangentWeight());
-					}
-					else
-					{
-						float PixelLength = 60.0f; 
-						ArriveTangentPoint.ScreenPosition = Key.ScreenPosition + FCurveEditor::GetVectorFromSlopeAndLength(ArriveTangent * -DisplayRatio, -PixelLength);
-					}
-					ArriveTangentPoint.LineDelta = Key.ScreenPosition - ArriveTangentPoint.ScreenPosition;
-					ArriveTangentPoint.LayerBias = 1;
-
-					Params.Points.Add(ArriveTangentPoint);
-				}
-
-				if (bShowTangents && Attributes.HasLeaveTangent())
-				{
-					float LeaveTangent = Attributes.GetLeaveTangent();
-
-					FCurvePointInfo LeaveTangentPoint(KeyHandle);
-					LeaveTangentPoint.Type = ECurvePointType::LeaveTangent;
-
-					if (Attributes.HasTangentWeightMode() && Attributes.HasLeaveTangentWeight() &&
-						(Attributes.GetTangentWeightMode() == RCTWM_WeightedBoth || Attributes.GetTangentWeightMode() == RCTWM_WeightedLeave))
-					{
-						LeaveTangentPoint.ScreenPosition = GetTangentPositionInScreenSpace(Key.ScreenPosition, LeaveTangent, Attributes.GetLeaveTangentWeight());
-					}
-					else
-					{
-						float PixelLength = 60.0f; 
-						LeaveTangentPoint.ScreenPosition = Key.ScreenPosition + FCurveEditor::GetVectorFromSlopeAndLength(LeaveTangent * -DisplayRatio, PixelLength);
-
-					}
-
-					LeaveTangentPoint.LineDelta = Key.ScreenPosition - LeaveTangentPoint.ScreenPosition;
-					LeaveTangentPoint.LayerBias = 1;
-
-					Params.Points.Add(LeaveTangentPoint);
-				}
-			}
-		}
-
-		OutDrawParams.Add(MoveTemp(Params));
-	}
+	InputMax = FMath::Max(InputMax, InputMin + 1e-10);
+	return FCurveEditorScreenSpaceH(PanelWidth, InputMin, InputMax);
 }
 
 void FCurveEditor::ConstructXGridLines(TArray<float>& MajorGridLines, TArray<float>& MinorGridLines, TArray<FText>& MajorGridLabels) const
 {
-	FCurveEditorScreenSpace ScreenSpace = GetScreenSpace();
+	FCurveEditorScreenSpaceH InputSpace = GetPanelInputSpace();
 
 	double MajorGridStep  = 0.0;
 	int32  MinorDivisions = 0;
-	if (GetSnapMetrics().InputSnapRate.ComputeGridSpacing(ScreenSpace.PixelsPerInput(), MajorGridStep, MinorDivisions))
+	if (GetSnapMetrics().InputSnapRate.ComputeGridSpacing(InputSpace.PixelsPerInput(), MajorGridStep, MinorDivisions))
 	{
-		const double FirstMajorLine = FMath::FloorToDouble(ScreenSpace.GetInputMin() / MajorGridStep) * MajorGridStep;
-		const double LastMajorLine  = FMath::CeilToDouble(ScreenSpace.GetInputMax() / MajorGridStep) * MajorGridStep;
+		FText GridLineLabelFormatX = GridLineLabelFormatXAttribute.Get();
+		const double FirstMajorLine = FMath::FloorToDouble(InputSpace.GetInputMin() / MajorGridStep) * MajorGridStep;
+		const double LastMajorLine  = FMath::CeilToDouble(InputSpace.GetInputMax() / MajorGridStep) * MajorGridStep;
 
 		for (double CurrentMajorLine = FirstMajorLine; CurrentMajorLine < LastMajorLine; CurrentMajorLine += MajorGridStep)
 		{
-			MajorGridLines.Add( ScreenSpace.SecondsToScreen(CurrentMajorLine) );
-			MajorGridLabels.Add( FText::Format(LOCTEXT("GridXLabelFormat", "{0}s"), FText::AsNumber(CurrentMajorLine)) );
+			MajorGridLines.Add( (CurrentMajorLine - InputSpace.GetInputMin()) * InputSpace.PixelsPerInput() );
+			MajorGridLabels.Add( FText::Format(GridLineLabelFormatX, FText::AsNumber(CurrentMajorLine)) );
 
 			for (int32 Step = 1; Step < MinorDivisions; ++Step)
 			{
-				MinorGridLines.Add( ScreenSpace.SecondsToScreen(CurrentMajorLine + Step*MajorGridStep/MinorDivisions) );
+				float MinorLine = CurrentMajorLine + Step*MajorGridStep/MinorDivisions;
+				MinorGridLines.Add( (MinorLine - InputSpace.GetInputMin()) * InputSpace.PixelsPerInput() );
 			}
 		}
 	}
 }
 
-void FCurveEditor::ConstructYGridLines(TArray<float>& MajorGridLines, TArray<float>& MinorGridLines, TArray<FText>&MajorGridLabels, uint8 MinorDivisions) const
-{
-	FCurveEditorSnapMetrics SnapMetrics = GetSnapMetrics();
-	FCurveEditorScreenSpace ScreenSpace = GetScreenSpace();
-
-	if (ScreenSpace.GetOutputMin() == ScreenSpace.GetOutputMax() || ScreenSpace.PixelsPerOutput() <= 0)
-	{
-		return;
-	}
-
-	const float GridPixelSpacing = FMath::Min(ScreenSpace.GetPhysicalHeight()/1.5f, 150.0f);
-
-	static float Base = 10;
-	float MaxTimeStep = GridPixelSpacing / ScreenSpace.PixelsPerOutput();
-	float MajorGridStep = FMath::Pow(Base, FMath::FloorToFloat(FMath::LogX(Base, MaxTimeStep)));
-
-	const double FirstMajorLine = FMath::FloorToDouble(ScreenSpace.GetOutputMin() / MajorGridStep) * MajorGridStep;
-	const double LastMajorLine  = FMath::CeilToDouble(ScreenSpace.GetOutputMax() / MajorGridStep) * MajorGridStep;
-
-	for (double CurrentMajorLine = FirstMajorLine; CurrentMajorLine < LastMajorLine; CurrentMajorLine += MajorGridStep)
-	{
-		MajorGridLines.Add( ScreenSpace.ValueToScreen(CurrentMajorLine) );
-		MajorGridLabels.Add(FText::Format(LOCTEXT("GridYLabelFormat", "{0}"), FText::AsNumber(CurrentMajorLine)));
-
-		for (int32 Step = 1; Step < MinorDivisions; ++Step)
-		{
-			MinorGridLines.Add( ScreenSpace.ValueToScreen(CurrentMajorLine + Step*MajorGridStep/MinorDivisions) );
-		}
-	}
-}
-
 void FCurveEditor::DeleteSelection()
 {
 	FScopedTransaction Transaction(LOCTEXT("DeleteKeys", "Delete Keys"));
@@ -725,179 +875,308 @@ void FCurveEditor::StraightenSelection()
 	}
 }
 
-void FCurveEditor::BakeSelection()
+bool FCurveEditor::CanFlattenOrStraightenSelection() const
 {
-	FText TransactionText = FText::Format(LOCTEXT("BakeCurves", "Bake {0}|plural(one=Curve, other=Curves)"), Selection.GetAll().Num());
-	FScopedTransaction Transaction(TransactionText);
-
-	TArray<FKeyHandle> KeyHandles;
-	TArray<FKeyPosition> SelectedKeyPositions;
+	return Selection.Count() > 0;
+}
 
-	TArray<FKeyPosition> NewKeyPositions;
-	TArray<FKeyAttributes> NewKeyAttributes;
+bool FCurveEditor::IsToolActive(const FCurveEditorToolID InToolID) const
+{
+	if (ActiveTool.IsSet())
+	{
+		return ActiveTool == InToolID;
+	}
 
-	FFrameRate BakeRate = GetSnapMetrics().InputSnapRate;
+	return false;
+}
 
-	for (const TTuple<FCurveModelID, FKeyHandleSet>& Pair : Selection.GetAll())
+void FCurveEditor::MakeToolActive(const FCurveEditorToolID InToolID)
+{
+	if (ActiveTool.IsSet())
 	{
-		if (FCurveModel* Curve = FindCurve(Pair.Key))
+		// Early out in the event that they're trying to switch to the same tool. This avoids
+		// unwanted activation/deactivation calls.
+		if (ActiveTool == InToolID)
 		{
-			KeyHandles.Reset(Pair.Value.Num());
-			KeyHandles.Append(Pair.Value.AsArray().GetData(), Pair.Value.Num());
+			return;
+		}
 
-			// Get all the selected keys
-			SelectedKeyPositions.SetNum(KeyHandles.Num());
-			Curve->GetKeyPositions(KeyHandles, SelectedKeyPositions);
+		// Deactivate the current tool before we activate the new one.
+		ToolExtensions[ActiveTool.GetValue()]->OnToolDeactivated();
+	}
 
-			// Find the hull of the range of the selected keys
-			double MinKey = TNumericLimits<double>::Max(), MaxKey = TNumericLimits<double>::Lowest();
-			for (FKeyPosition Key : SelectedKeyPositions)
-			{
-				MinKey = FMath::Min(Key.InputValue, MinKey);
-				MaxKey = FMath::Max(Key.InputValue, MaxKey);
-			}
+	ActiveTool.Reset();
 
-			// Get all keys that exist between the time range
-			KeyHandles.Reset();
-			Curve->GetKeys(*this, MinKey, MaxKey, -MAX_dbl, MAX_dbl, KeyHandles);
+	// Notify anyone listening that we've switched tools (possibly to an inactive one)
+	OnActiveToolChangedDelegate.Broadcast(InToolID);
 
-			if (KeyHandles.Num() > 1)
-			{
-				// Determine new times for new keys
-				double Interval = BakeRate.AsInterval();
-				int32 NumKeysToAdd = FMath::FloorToInt((MaxKey - MinKey) / Interval);
+	if (InToolID != FCurveEditorToolID::Unset())
+	{
+		ActiveTool = InToolID;
+		ToolExtensions[ActiveTool.GetValue()]->OnToolActivated();
+	}
+}
 
-				NewKeyPositions.Reset(NumKeysToAdd);
-				NewKeyAttributes.Reset(NumKeysToAdd);
+ICurveEditorToolExtension* FCurveEditor::GetCurrentTool() const
+{
+	if (ActiveTool.IsSet())
+	{
+		return ToolExtensions[ActiveTool.GetValue()].Get();
+	}
 
-				for (int32 KeyIndex = 0; KeyIndex < NumKeysToAdd; ++KeyIndex)
-				{
-					FKeyPosition CurrentKey(MinKey + KeyIndex*Interval, 0.0);
-					if (Curve->Evaluate(CurrentKey.InputValue, CurrentKey.OutputValue))
-					{
-						NewKeyPositions.Add(CurrentKey);
-						NewKeyAttributes.Add(FKeyAttributes().SetInterpMode(RCIM_Linear));
-					}
-				}
+	// If there is no active tool we return nullptr.
+	return nullptr;
+}
 
-				Curve->Modify();
+TSet<FCurveModelID> FCurveEditor::GetEditedCurves() const
+{
+	TArray<FCurveModelID> AllCurves;
+	GetCurves().GenerateKeyArray(AllCurves);
+	return TSet<FCurveModelID>(AllCurves);
+}
 
-				// Remove all the old keys and add the new ones
-				Curve->RemoveKeys(KeyHandles);
-				Curve->AddKeys(NewKeyPositions, NewKeyAttributes);
-			}
-		}
+void FCurveEditor::SetBufferedCurves(const TSet<FCurveModelID>& InCurves)
+{
+	BufferedCurves.Empty();
+
+	// We make a copy of the curve data and store it.
+	for (FCurveModelID CurveID : InCurves)
+	{
+		FCurveModel* CurveModel = FindCurve(CurveID);
+		check(CurveModel);
+
+
+		// We loop through the key data on the curve and apply it to the copy, since there's no copy constructors for Curve Models (for good reason!)
+		FBufferedCurve& BufferedCurve = BufferedCurves.Emplace_GetRef();
+		
+		int32 NumKeys = CurveModel->GetNumKeys();
+		BufferedCurve.KeyPositions.SetNumUninitialized(NumKeys);
+		BufferedCurve.KeyAttributes.SetNumUninitialized(NumKeys);
+
+		// Copy all of the from the curve into our cached buffer
+		TArray<FKeyHandle> KeyHandles;
+
+		CurveModel->GetKeys(*this, TNumericLimits<double>::Lowest(), TNumericLimits<double>::Max(), TNumericLimits<double>::Lowest(), TNumericLimits<double>::Max(), KeyHandles);
+		CurveModel->GetKeyPositions(KeyHandles, BufferedCurve.KeyPositions);
+		CurveModel->GetKeyAttributes(KeyHandles, BufferedCurve.KeyAttributes);
+		BufferedCurve.IntentionName = CurveModel->GetIntentionName();
 	}
 }
 
-/**
-  The following key reduction is the same as that found in FRichCurve.
-  It would be nice if there was just one implementation of the reduction (and) baking algorithms.
-*/
-/** Util to find float value on bezier defined by 4 control points */
-static float BezierInterp(float P0, float P1, float P2, float P3, float Alpha)
+
+void FCurveEditor::ApplyBufferedCurveToTarget(const FBufferedCurve& BufferedCurve, FCurveModel* TargetCurve)
 {
-	const float P01 = FMath::Lerp(P0, P1, Alpha);
-	const float P12 = FMath::Lerp(P1, P2, Alpha);
-	const float P23 = FMath::Lerp(P2, P3, Alpha);
-	const float P012 = FMath::Lerp(P01, P12, Alpha);
-	const float P123 = FMath::Lerp(P12, P23, Alpha);
-	const float P0123 = FMath::Lerp(P012, P123, Alpha);
+	check(TargetCurve);
+
+	TArray<FKeyPosition> KeyPositions;
+	TArray<FKeyAttributes> KeyAttributes;
+
+
+	// Copy the data from the Buffered curve into the target curve. This just does wholesale replacement.
+	TArray<FKeyHandle> TargetKeyHandles;
+	TargetCurve->GetKeys(*this, TNumericLimits<double>::Lowest(), TNumericLimits<double>::Max(), TNumericLimits<double>::Lowest(), TNumericLimits<double>::Max(), TargetKeyHandles);
 
-	return P0123;
+	// Clear our current keys from the target curve
+	TargetCurve->RemoveKeys(TargetKeyHandles);
+
+	// Now put our buffered keys into the target curve
+	TargetCurve->AddKeys(BufferedCurve.KeyPositions, BufferedCurve.KeyAttributes);
 }
 
-static float EvalForTwoKeys (const FKeyPosition& Key1Pos, const FKeyAttributes& Key1Attrib,
-							  const FKeyPosition& Key2Pos, const FKeyAttributes& Key2Attrib,
-							  const float InTime)
+bool FCurveEditor::ApplyBufferedCurves(const TSet<FCurveModelID>& InCurvesToApplyTo)
 {
-	const float Diff = Key2Pos.InputValue - Key1Pos.InputValue;
-
-	if (Diff > 0.f && Key1Attrib.GetInterpMode() != RCIM_Constant)
+	// Each curve can specify an "Intention" name. This gives a little bit of context about how the curve is intended to be used,
+	// without locking anyone into a specific set of intentions. When you go to apply the buffered curves, for each curve that you
+	// want to apply it to, we can look in our stored curves to see if someone has the same intention. If there isn't a matching intention
+	// then we skip and consider a fallback method (such as 1:1 copy). There is a lot of guessing still involved as there are complex
+	// situations that users may try to use it in (such as buffering two sets of transform curves and applying it to two destination transform curves)
+	// or trying to copy something with a name like "Focal Length" and pasting it onto a different track. We don't handle these cases for now,
+	// but attempt to communicate it to the user via  toast notification when pasting fails for whatever reason.
+	int32 NumCurvesMatchedByIntent = 0;
+	int32 NumCurvesNoMatchedIntent = 0;
+	bool bFoundAnyMatchedIntent = false;
+
+	TMap<FString, int32> IntentMatchIndexes;
+
+	for (const FCurveModelID CurveModelID : InCurvesToApplyTo)
 	{
-		const float Alpha = (InTime - Key1Pos.InputValue) / Diff;
-		const float P0 = Key1Pos.OutputValue;
-		const float P3 = Key2Pos.OutputValue;
+		FCurveModel* TargetCurve = FindCurve(CurveModelID);
+		check(TargetCurve);
+
+		// Figure out what our destination thinks it's supposed to be used for, ie "Location.X"
+		FString TargetIntent = TargetCurve->GetIntentionName();
+		if (TargetIntent.IsEmpty())
+		{
+			// We don't try to match curves with no intent as that's just chaos.
+			NumCurvesNoMatchedIntent++;
+			continue;
+		}
+
+		TargetCurve->Modify();
+
+		// In an attempt to support buffering multiple curves with the same intention, we'll try to match them up in pairs. This means
+		// for the first curve that we're trying to apply to, if the intention is "Location.X" we will search the buffered curves for a
+		// "Location.X". Upon finding one, we store the index that it was found at, so the next time we try to find a curve with the same
+		// intention, we look for the second "Location.X" and so forth. If we don't find a second "Location.X" in our buffered curves we'll
+		// fall back to the first buffered one so you can 1:Many copy a curve.
+		int32 BufferedCurveSearchIndexStart = 0;
+		const int32* PreviouslyFoundIntent = IntentMatchIndexes.Find(TargetIntent);
+		if (PreviouslyFoundIntent)
+		{
+			// Start our search on the next item in the array. If we don't find one, we'll fall back to the last one.
+			BufferedCurveSearchIndexStart = IntentMatchIndexes[TargetIntent] + 1;
+		}
+
+		int32 MatchedBufferedCurveIndex = -1;
+		for (int32 BufferedCurveIndex = BufferedCurveSearchIndexStart; BufferedCurveIndex < BufferedCurves.Num(); BufferedCurveIndex++)
+		{
+			if (BufferedCurves[BufferedCurveIndex].IntentionName == TargetIntent)
+			{
+				MatchedBufferedCurveIndex = BufferedCurveIndex;
+
+				// Update our previously found intent to the latest one.
+				IntentMatchIndexes.FindOrAdd(TargetIntent) = MatchedBufferedCurveIndex;
+				break;
+			}
+		}
 
-		if (Key1Attrib.GetInterpMode() == RCIM_Linear)
+		// The Intent Match Indexes stores the latest index to find a valid curve, or the last one if no new valid one was found.
+		// If there is an entry in the match indexes now, we can use that to figure out which buffered curve we'll pull from.
+		// If we didn't find any more with the same intention, we fall back to the existing one (if it exists!)
+		if (IntentMatchIndexes.Find(TargetIntent))
 		{
-			return FMath::Lerp(P0, P3, Alpha);
+			MatchedBufferedCurveIndex = IntentMatchIndexes[TargetIntent];
+		}
+
+		// Finally, we can try to use the matched curve if one was found.
+		if (MatchedBufferedCurveIndex >= 0)
+		{
+			// We successfully matched, so count that one up!
+			NumCurvesMatchedByIntent++;
+			bFoundAnyMatchedIntent = true;
+
+			const FBufferedCurve& BufferedCurve = BufferedCurves[MatchedBufferedCurveIndex];
+			ApplyBufferedCurveToTarget(BufferedCurve, TargetCurve);
+
 		}
 		else
 		{
-			const float OneThird = 1.0f / 3.0f;
-			const float P1 = P0 + (Key1Attrib.GetLeaveTangent() * Diff*OneThird);
-			const float P2 = P3 - (Key2Attrib.GetArriveTangent() * Diff*OneThird);
+			// We couldn't find a match despite our best efforts
+			NumCurvesNoMatchedIntent++;
+		}
+	}
+
+	// If we managed to match any by intent, we're going to early out and assume that's what their intent was.
+	if (bFoundAnyMatchedIntent)
+	{
+		const FText NotificationText = FText::Format(LOCTEXT("MatchedBufferedCurvesByIntent", "Applied {0}/{1} buffered curves to {2}/{3} target curves."),
+			FText::AsNumber(IntentMatchIndexes.Num()), FText::AsNumber(BufferedCurves.Num()),		// We used X of Y total buffered curves
+			FText::AsNumber(NumCurvesMatchedByIntent), FText::AsNumber(InCurvesToApplyTo.Num()));	// To apply to Z of W target curves,
+
+		FNotificationInfo Info(NotificationText);
+		Info.ExpireDuration = 6.f;
+		Info.bUseLargeFont = false;
+		Info.bUseSuccessFailIcons = false;
+		FSlateNotificationManager::Get().AddNotification(Info);
 
-			return BezierInterp(P0, P1, P2, P3, Alpha);
+		if (NumCurvesNoMatchedIntent > 0)
+		{
+			const FText FailedNotificationText = FText::Format(LOCTEXT("NumCurvesNotMatchedByIntent", "Failed to find a buffered curve with the same intent for {0} target curves, skipping..."),
+				FText::AsNumber(NumCurvesNoMatchedIntent));												// Leaving V many target curves unaffected due to no intent match.
+
+			FNotificationInfo FailInfo(FailedNotificationText);
+			FailInfo.ExpireDuration = 6.f;
+			FailInfo.bUseLargeFont = false;
+			FailInfo.bUseSuccessFailIcons = true;
+			FSlateNotificationManager::Get().AddNotification(FailInfo);
 		}
+
+		// Early out
+		return true;
 	}
-	else
+
+	// If we got this far, it means that the buffered curves have no recognizable relation to the target curves.
+	// If the number of curves match, we'll just do a 1:1 mapping. This works for most cases where you're trying
+	// to paste an unrelated curve onto another as it's likely that there's only one curve. We don't limit it to
+	// one curve though, we'll just warn...
+	if (InCurvesToApplyTo.Num() == BufferedCurves.Num())
 	{
-		return Key1Pos.OutputValue;
+		// This will work great in the case there's only one curve. It'll guess if there's more than one, relying on
+		// sets with no guaranteed order.
+		TArray<FCurveModelID> CurvesToApplyTo = InCurvesToApplyTo.Array();
+		
+		for (int32 CurveIndex = 0; CurveIndex < InCurvesToApplyTo.Num(); CurveIndex++)
+		{
+			ApplyBufferedCurveToTarget(BufferedCurves[CurveIndex], FindCurve(CurvesToApplyTo[CurveIndex]));
+		}
+
+		FText NotificationText;
+		if (InCurvesToApplyTo.Num() == 1)
+		{
+			NotificationText = LOCTEXT("MatchedBufferedCurvesBySolo", "Applied buffered curve to target curve with no intention matching.");
+		}
+		else
+		{
+			NotificationText = LOCTEXT("MatchedBufferedCurvesByIndex", "Applied buffered curves with no intention matching. Order not guranteed.");
+		}
+
+		FNotificationInfo Info(NotificationText);
+		Info.ExpireDuration = 6.f;
+		Info.bUseLargeFont = false;
+		Info.bUseSuccessFailIcons = false;
+		FSlateNotificationManager::Get().AddNotification(Info);
+
+		// Early out
+		return true;
 	}
-}
 
+	// If we got this far, we have no idea what to do. They're trying to match a bunch of curves with no intention and different amounts. 
+	// Warn of failure and give up.
+	{
+		const FText FailedNotificationText = LOCTEXT("NoBufferedCurvesMatched", "Failed to apply buffered curves, apply them one at a time instead.");
 
-void FCurveEditor::SimplifySelection(float Tolerance)
-{
-	FText TransactionText = FText::Format(LOCTEXT("SimplifyCurves", "Simplify {0}|plural(one=Curve, other=Curves)"), Selection.GetAll().Num());
-	FScopedTransaction Transaction(TransactionText);
-	TArray<FKeyHandle> KeyHandles;
-	TArray<FKeyPosition> SelectedKeyPositions;
-	TArray<FKeyAttributes> SelectedKeyAttributes;
+		FNotificationInfo FailInfo(FailedNotificationText);
+		FailInfo.ExpireDuration = 6.f;
+		FailInfo.bUseLargeFont = false;
+		FailInfo.bUseSuccessFailIcons = true;
+		FSlateNotificationManager::Get().AddNotification(FailInfo);
+	}
 
-	FFrameRate BakeRate = GetSnapMetrics().InputSnapRate;
+	return false;
+}
 
-	for (const TTuple<FCurveModelID, FKeyHandleSet>& Pair : Selection.GetAll())
+void FCurveEditor::PostUndo(bool bSuccess)
+{
+	// If you create keys and then undo them the selection set still thinks there's keys selected.
+	// This presents issues with context menus and other things that are activated when there is a selection set.
+	// To fix this, we have to loop through all of our curve models, and re-select only the key handles that were
+	// previously selected that still exist. Ugly, but reasonably functional.
+	TMap<FCurveModelID, FKeyHandleSet> SelectionSet = Selection.GetAll();
+	for (const TPair<FCurveModelID, FKeyHandleSet>& Set : SelectionSet)
 	{
-		if (FCurveModel* Curve = FindCurve(Pair.Key))
-		{
-			KeyHandles.Reset(Pair.Value.Num());
-			KeyHandles.Append(Pair.Value.AsArray().GetData(), Pair.Value.Num());
+		FCurveModel* CurveModel = FindCurve(Set.Key);
 
-			// Get all the selected keys
-			SelectedKeyPositions.SetNum(KeyHandles.Num());
-			Curve->GetKeyPositions(KeyHandles, SelectedKeyPositions);
+		// If the entire curve was removed, just dump that out of the selection set.
+		if (!CurveModel)
+		{
+			Selection.Remove(Set.Key);
+			continue;
+		}
 
-			// Find the hull of the range of the selected keys
-			double MinKey = TNumericLimits<double>::Max(), MaxKey = TNumericLimits<double>::Lowest();
-			for (FKeyPosition Key : SelectedKeyPositions)
-			{
-				MinKey = FMath::Min(Key.InputValue, MinKey);
-				MaxKey = FMath::Max(Key.InputValue, MaxKey);
-			}
+		// Get all of the key handles from this curve.
+		TArray<FKeyHandle> KeyHandles;
+		CurveModel->GetKeys(*this, TNumericLimits<double>::Lowest(), TNumericLimits<double>::Max(), TNumericLimits<double>::Lowest(), TNumericLimits<double>::Max(), KeyHandles);
 
-			// Get all keys that exist between the time range
-			KeyHandles.Reset();
-			Curve->GetKeys(*this, MinKey, MaxKey, -MAX_dbl, MAX_dbl, KeyHandles);
-			if (KeyHandles.Num() > 2) //need at least 3 keys to reduce
+		// The set handles will be mutated as we remove things so we need a copy that we can iterate through.
+		TArrayView<const FKeyHandle> SelectedHandles = Set.Value.AsArray();
+		TArray<FKeyHandle> NonMutableArray = TArray<FKeyHandle>(SelectedHandles.GetData(), SelectedHandles.Num());
+		
+		for (const FKeyHandle& Handle : NonMutableArray)
+		{
+			// Check to see if our curve model contains this handle still.
+			if (!KeyHandles.Contains(Handle))
 			{
-				SelectedKeyPositions.SetNum(KeyHandles.Num());
-				Curve->GetKeyPositions(KeyHandles, SelectedKeyPositions);
-				SelectedKeyAttributes.SetNum(KeyHandles.Num());
-				Curve->GetKeyAttributes(KeyHandles, SelectedKeyAttributes);
-			
-				int32 MostRecentKeepKeyIndex = 0;
-				TArray<FKeyHandle> KeysToRemove;
-				for (int32 TestIndex = 1; TestIndex < KeyHandles.Num()-1; ++TestIndex)
-				{
-					
-					const float KeyValue = SelectedKeyPositions[TestIndex].OutputValue;
-					const float ValueWithoutKey = EvalForTwoKeys(SelectedKeyPositions[MostRecentKeepKeyIndex], SelectedKeyAttributes[MostRecentKeepKeyIndex],
-						SelectedKeyPositions[TestIndex + 1], SelectedKeyAttributes[TestIndex + 1],
-						SelectedKeyPositions[TestIndex].InputValue);
-					if (FMath::Abs(ValueWithoutKey - KeyValue) > Tolerance) // Is this key needed
-					{
-						MostRecentKeepKeyIndex = TestIndex;
-					}
-					else
-					{
-						KeysToRemove.Add(KeyHandles[TestIndex]);
-					}
-					}
-				Curve->Modify();
-				Curve->RemoveKeys(KeysToRemove);
+				Selection.Remove(Set.Key, ECurvePointType::Key, Handle);
 			}
 		}
 	}
diff --git a/Engine/Source/Editor/CurveEditor/Private/CurveEditorCommands.cpp b/Engine/Source/Editor/CurveEditor/Private/CurveEditorCommands.cpp
--- a/Engine/Source/Editor/CurveEditor/Private/CurveEditorCommands.cpp
+++ b/Engine/Source/Editor/CurveEditor/Private/CurveEditorCommands.cpp
@@ -11,33 +11,37 @@ void FCurveEditorCommands::RegisterCommands()
 	UI_COMMAND(ZoomToFit, "Fit", "Zoom to Fit", EUserInterfaceActionType::Button, FInputChord(EKeys::F));
 	UI_COMMAND(ZoomToFitAll, "FitAll", "Zoom to Fit All", EUserInterfaceActionType::Button, FInputChord(EKeys::A));
 
-	UI_COMMAND(ToggleInputSnapping, "Input Snapping", "Toggle Input Snapping", EUserInterfaceActionType::ToggleButton, FInputChord());
-	UI_COMMAND(ToggleOutputSnapping, "Output Snapping", "Toggle Output Snapping", EUserInterfaceActionType::ToggleButton, FInputChord());
+	UI_COMMAND(ToggleInputSnapping, "Input Snapping", "Toggle Time Snapping", EUserInterfaceActionType::ToggleButton, FInputChord());
+	UI_COMMAND(ToggleOutputSnapping, "Output Snapping", "Toggle Value Snapping", EUserInterfaceActionType::ToggleButton, FInputChord());
 
 	UI_COMMAND(InterpolationConstant, "Constant", "Constant interpolation", EUserInterfaceActionType::ToggleButton, FInputChord(EKeys::Five));
 	UI_COMMAND(InterpolationLinear, "Linear", "Linear interpolation", EUserInterfaceActionType::ToggleButton, FInputChord(EKeys::Four));
 	UI_COMMAND(InterpolationCubicAuto, "Auto", "Cubic interpolation - Automatic tangents", EUserInterfaceActionType::ToggleButton, FInputChord(EKeys::One));
 	UI_COMMAND(InterpolationCubicUser, "User", "Cubic interpolation - User flat tangents", EUserInterfaceActionType::ToggleButton, FInputChord(EKeys::Two));
 	UI_COMMAND(InterpolationCubicBreak, "Break", "Cubic interpolation - User broken tangents", EUserInterfaceActionType::ToggleButton, FInputChord(EKeys::Three));
-	UI_COMMAND(InterpolationToggleWeighted, "Weighted Tangents", "Toggle weighted tangents for cubic interpolation modes", EUserInterfaceActionType::ToggleButton, FInputChord(EModifierKey::Control, EKeys::W));
+	UI_COMMAND(InterpolationToggleWeighted, "Weighted Tangents", "Toggle weighted tangents for cubic interpolation modes. Only supported on some curve types", EUserInterfaceActionType::ToggleButton, FInputChord(EModifierKey::Control, EKeys::W));
 
-	UI_COMMAND(FlattenTangents, "Flatten", "Flatten Tangents", EUserInterfaceActionType::Button, FInputChord());
+	UI_COMMAND(FlattenTangents, "Flatten", "Flatten Tangents", EUserInterfaceActionType::Button, FInputChord(EKeys::S));
 	UI_COMMAND(StraightenTangents, "Straighten", "Straighten Tangents", EUserInterfaceActionType::Button, FInputChord());
 
 	UI_COMMAND(BakeCurve, "Bake", "Bake curve", EUserInterfaceActionType::Button, FInputChord());
 	UI_COMMAND(ReduceCurve, "Reduce", "Reduce curve", EUserInterfaceActionType::Button, FInputChord());
 
-	UI_COMMAND(SetPreInfinityExtrapCycle, "Cycle", "Set the pre-infinity curve extrapolation to cycle.", EUserInterfaceActionType::RadioButton, FInputChord());
-	UI_COMMAND(SetPreInfinityExtrapCycleWithOffset, "Cycle with Offset", "Set the pre-infinity curve extrapolation to cycle with offset.", EUserInterfaceActionType::RadioButton, FInputChord());
-	UI_COMMAND(SetPreInfinityExtrapOscillate, "Oscillate", "Set the pre-infinity curve extrapolation to oscillate.", EUserInterfaceActionType::RadioButton, FInputChord());
-	UI_COMMAND(SetPreInfinityExtrapLinear, "Linear", "Set the pre-infinity curve extrapolation to linear.", EUserInterfaceActionType::RadioButton, FInputChord());
-	UI_COMMAND(SetPreInfinityExtrapConstant, "Constant", "Set the pre-infinity curve extrapolation to constant.", EUserInterfaceActionType::RadioButton, FInputChord());
-	UI_COMMAND(SetPostInfinityExtrapCycle, "Cycle", "Set the post-infinity curve extrapolation to cycle.", EUserInterfaceActionType::RadioButton, FInputChord());
-	UI_COMMAND(SetPostInfinityExtrapCycleWithOffset, "Cycle with Offset", "Set the post-infinity curve extrapolation to cycle with offset.", EUserInterfaceActionType::RadioButton, FInputChord());
-	UI_COMMAND(SetPostInfinityExtrapOscillate, "Oscillate", "Set the post-infinity curve extrapolation to oscillate.", EUserInterfaceActionType::RadioButton, FInputChord());
-	UI_COMMAND(SetPostInfinityExtrapLinear, "Linear", "Set the post-infinity curve extrapolation to linear.", EUserInterfaceActionType::RadioButton, FInputChord());
-	UI_COMMAND(SetPostInfinityExtrapConstant, "Constant", "Set the post-infinity curve extrapolation to constant.", EUserInterfaceActionType::RadioButton, FInputChord());
+	// Pre and Post Infinity
+	UI_COMMAND(SetPreInfinityExtrapCycle, "Cycle", "Cycle creates a repeating cycle from the first to last key, effectively modulating the input time. This can create jumps if the terminus values are not the same value.", EUserInterfaceActionType::RadioButton, FInputChord());
+	UI_COMMAND(SetPreInfinityExtrapCycleWithOffset, "Cycle with Offset", "Creates a repeating cycle where the value is added to the last value of the previous Cycle. This will avoid jumps but will cause drift over time if there's a net change in value.", EUserInterfaceActionType::RadioButton, FInputChord());
+	UI_COMMAND(SetPreInfinityExtrapOscillate, "Oscillate (Ping Pong)", "Creates a repeating cycle which ping pongs and will play from beginning to end, then end to beginning.", EUserInterfaceActionType::RadioButton, FInputChord());
+	UI_COMMAND(SetPreInfinityExtrapLinear, "Linear", "Linearly interpolates based on the in tangent of the first key.", EUserInterfaceActionType::RadioButton, FInputChord());
+	UI_COMMAND(SetPreInfinityExtrapConstant, "Constant", "Extrapolation will always return the value of the first key.", EUserInterfaceActionType::RadioButton, FInputChord());
 
+	UI_COMMAND(SetPostInfinityExtrapCycle, "Cycle", "Cycle creates a repeating cycle from the first to last key, effectively modulating the input time. This can create jumps if the terminus values are not the same value", EUserInterfaceActionType::RadioButton, FInputChord());
+	UI_COMMAND(SetPostInfinityExtrapCycleWithOffset, "Cycle with Offset", "Creates a repeating cycle where the value is added to the last value of the previous Cycle. This will avoid jumps but will cause drift over time if there's a net change in value.", EUserInterfaceActionType::RadioButton, FInputChord());
+	UI_COMMAND(SetPostInfinityExtrapOscillate, "Oscillate (Ping Pong)", "Creates a repeating cycle which ping pongs and will play from beginning to end, then end to beginning.", EUserInterfaceActionType::RadioButton, FInputChord());
+	UI_COMMAND(SetPostInfinityExtrapLinear, "Linear", "Linearly interpolates based on the in tangent of the last key.", EUserInterfaceActionType::RadioButton, FInputChord());
+	UI_COMMAND(SetPostInfinityExtrapConstant, "Constant", "Extrapolation will always return the value of the last key.", EUserInterfaceActionType::RadioButton, FInputChord());
+
+
+	// Tangent Visibility
 	UI_COMMAND(SetAllTangentsVisibility, "All Tangents", "Show all tangents in the curve editor.", EUserInterfaceActionType::RadioButton, FInputChord());
 	UI_COMMAND(SetSelectedKeysTangentVisibility, "Selected Keys", "Show tangents for selected keys in the curve editor.", EUserInterfaceActionType::RadioButton, FInputChord());
 	UI_COMMAND(SetNoTangentsVisibility, "No Tangents", "Show no tangents in the curve editor.", EUserInterfaceActionType::RadioButton, FInputChord());
@@ -47,8 +51,39 @@ void FCurveEditorCommands::RegisterCommands()
 
 	UI_COMMAND(AddKeyHovered, "Add Key", "Add a new key to this curve at the current position.", EUserInterfaceActionType::Button, FInputChord(EKeys::MiddleMouseButton) );
 
-	UI_COMMAND(AddKeyToAllCurves, "Add Key", "Add a new key to all curves at the current time.", EUserInterfaceActionType::Button, FInputChord(EModifierKey::Shift, EKeys::LeftMouseButton) );
-	UI_COMMAND(AddKeyToAllCurvesHere, "Add Key Here", "Add a new key to all curves at the current time and value.", EUserInterfaceActionType::Button, FInputChord(EModifierKey::Shift | EModifierKey::Control, EKeys::LeftMouseButton) );
+	UI_COMMAND(AddKeyToAllCurves, "Add Key", "Add a new key to all curves at the current time.", EUserInterfaceActionType::Button, FInputChord(EKeys::Enter) );
+
+	// Graph Viewing Modes
+	UI_COMMAND(SetViewModeAbsolute, "Absolute View Mode", "Absolute view displays all curves overlapping with the Y axis proportionally scaled.", EUserInterfaceActionType::ToggleButton, FInputChord());
+	UI_COMMAND(SetViewModeStacked, "Stacked View Mode", "Stacked view displays each curve in its own graph with the Y axis normalized [-1, 1].", EUserInterfaceActionType::ToggleButton, FInputChord());
+	UI_COMMAND(SetViewModeNormalized, "Normalized View Mode", "Normalized view displays all curves overlapping with the Y axis normalized [-1, 1].", EUserInterfaceActionType::ToggleButton, FInputChord());
+
+	// Deactivate the currently active tool
+	UI_COMMAND(DeactivateCurrentTool, "Deactivate Tool", "Deactivates the current tool and returns to just supporting selection.", EUserInterfaceActionType::ToggleButton, FInputChord(EKeys::Q));
+
+	// User Implementable Filter window
+	UI_COMMAND(OpenUserImplementableFilterWindow, "Filter...", "Opens a window which lets you choose from user implementable filter classes with advanced settings.", EUserInterfaceActionType::Button, FInputChord());
+	
+	// Deselect any keys that the user has selected.
+	UI_COMMAND(DeselectAllKeys, "Deselect Keys", "Clears your current key selection.", EUserInterfaceActionType::Button, FInputChord(EModifierKey::Control, EKeys::D));
+
+	// Buffer and Apply Curves. Like copy and paste, but with multiple curve support.
+	UI_COMMAND(BufferVisibleCurves, "Store Curves", "Stores a copy of the visible curves which can be applied onto other curve sets.", EUserInterfaceActionType::Button, FInputChord());
+	// This name is overwritten in CurveEditorContextMenu to show the number of stashed curves.
+	UI_COMMAND(ApplyBufferedCurves, "Apply Stored Curves", "Applies the stored curves to the visible set.", EUserInterfaceActionType::Button, FInputChord());
+
+	// Axis Snapping
+	UI_COMMAND(SetAxisSnappingNone, "Both", "Disable axis snapping and allow movement on both the X and Y directions.", EUserInterfaceActionType::Button, FInputChord());
+	UI_COMMAND(SetAxisSnappingHorizontal, "X Only", "Snap transform tool axis movement to X direction.", EUserInterfaceActionType::Button, FInputChord());
+	UI_COMMAND(SetAxisSnappingVertical, "Y Only", "Snap transform tool axis movement to Y direction.", EUserInterfaceActionType::Button, FInputChord());	
+	
+	// Time Management
+	UI_COMMAND(StepToNextKey, "Step to Next Key", "Step to the next key", EUserInterfaceActionType::Button, FInputChord(EKeys::Period));
+	UI_COMMAND(StepToPreviousKey, "Step to Previous Key", "Step to the previous key", EUserInterfaceActionType::Button, FInputChord(EKeys::Comma));
+	UI_COMMAND(StepForward, "Step Forward", "Step the timeline forward", EUserInterfaceActionType::Button, FInputChord(EKeys::Right));
+	UI_COMMAND(StepBackward, "Step Backward", "Step the timeline backward", EUserInterfaceActionType::Button, FInputChord(EKeys::Left));
+	UI_COMMAND(JumpToStart, "Jump to Start", "Jump to the start of the playback range", EUserInterfaceActionType::Button, FInputChord(EKeys::Up));
+	UI_COMMAND(JumpToEnd, "Jump to End", "Jump to the end of the playback range", EUserInterfaceActionType::Button, FInputChord(EModifierKey::Control, EKeys::Up));
 }
 
 #undef LOCTEXT_NAMESPACE
diff --git a/Engine/Source/Editor/CurveEditor/Private/CurveEditorContextMenu.cpp b/Engine/Source/Editor/CurveEditor/Private/CurveEditorContextMenu.cpp
--- a/Engine/Source/Editor/CurveEditor/Private/CurveEditorContextMenu.cpp
+++ b/Engine/Source/Editor/CurveEditor/Private/CurveEditorContextMenu.cpp
@@ -8,37 +8,44 @@
 
 #define LOCTEXT_NAMESPACE "CurveEditorContextMenu"
 
-void FCurveEditorContextMenu::BuildMenu(FMenuBuilder& MenuBuilder, TWeakPtr<FCurveEditor> WeakCurveEditor, TOptional<FCurvePointHandle> ClickedPoint, TOptional<FCurveModelID> HoveredCurveID)
+void FCurveEditorContextMenu::BuildMenu(FMenuBuilder& MenuBuilder, TSharedRef<FCurveEditor> CurveEditor, TOptional<FCurvePointHandle> ClickedPoint, TOptional<FCurveModelID> HoveredCurveID)
 {
-	TSharedPtr<FCurveEditor> CurveEditor = WeakCurveEditor.Pin();
-	if (!CurveEditor.IsValid())
-	{
-		return;
-	}
+	int32 NumSelectedKeys = CurveEditor->GetSelection().Count();
+
+	TSharedPtr<FCurveEditor> LocalCurveEditor = CurveEditor->AsShared();
+
+	// We change the name to reflect the current number of curves selected.
+	TAttribute<FText> ApplyBufferedCurvesText = TAttribute<FText>::Create(TAttribute<FText>::FGetter::CreateLambda([LocalCurveEditor] {
+		return FText::Format(LOCTEXT("ApplyStoredCurvesContextMenu", "Apply {0} Stored Curves"), LocalCurveEditor->GetNumBufferedCurves());
+	}));
 
-	int32 NumSelectedKeys = CurveEditor->Selection.Count();
-	if (NumSelectedKeys > 0)
+	// We prioritize key selections (because otherwise you have to be exactly on a key to context menu)
+	// but this creates an issue where you have a selection but try to right click on a curve it still does
+	// the key one... so this lets us get the key menu anywhere on the graph except for over a curve which seems fair.
+	if (NumSelectedKeys > 0 && !HoveredCurveID.IsSet())
 	{
 		MenuBuilder.BeginSection("CurveEditorKeySection", FText::Format(LOCTEXT("CurveEditorKeySection", "{0} Selected {0}|plural(one=Key,other=Keys)"), NumSelectedKeys));
 		{
+			// Modify Data
 			MenuBuilder.AddMenuEntry(FGenericCommands::Get().Delete);
-
+			
 			MenuBuilder.AddMenuEntry(FCurveEditorCommands::Get().FlattenTangents);
 			MenuBuilder.AddMenuEntry(FCurveEditorCommands::Get().StraightenTangents);
-			MenuBuilder.AddMenuEntry(FCurveEditorCommands::Get().ReduceCurve);
-			MenuBuilder.AddMenuEntry(FCurveEditorCommands::Get().BakeCurve);
 
 			MenuBuilder.AddMenuSeparator();
 
+			// Tangent Types
 			MenuBuilder.AddMenuEntry(FCurveEditorCommands::Get().InterpolationCubicAuto);
 			MenuBuilder.AddMenuEntry(FCurveEditorCommands::Get().InterpolationCubicUser);
 			MenuBuilder.AddMenuEntry(FCurveEditorCommands::Get().InterpolationCubicBreak);
 			MenuBuilder.AddMenuEntry(FCurveEditorCommands::Get().InterpolationLinear);
 			MenuBuilder.AddMenuEntry(FCurveEditorCommands::Get().InterpolationConstant);
+			MenuBuilder.AddMenuEntry(FCurveEditorCommands::Get().InterpolationToggleWeighted);
 
 			MenuBuilder.AddMenuSeparator();
 
-			MenuBuilder.AddMenuEntry(FCurveEditorCommands::Get().InterpolationToggleWeighted);
+			// Filters
+			MenuBuilder.AddMenuEntry(FCurveEditorCommands::Get().OpenUserImplementableFilterWindow);
 		}
 		MenuBuilder.EndSection();
 	}
@@ -47,12 +54,15 @@ void FCurveEditorContextMenu::BuildMenu(FMenuBuilder& MenuBuilder, TWeakPtr<FCur
 		const FCurveModel* HoveredCurve = HoveredCurveID.IsSet() ? CurveEditor->FindCurve(HoveredCurveID.GetValue()) : nullptr;
 		if (HoveredCurve)
 		{
-			MenuBuilder.BeginSection("CurveEditorCurveSection", FText::Format(LOCTEXT("CurveNameFormat", "Curve '{0}'"), HoveredCurve->GetDisplayName()));
+			MenuBuilder.BeginSection("CurveEditorCurveSection", FText::Format(LOCTEXT("CurveNameFormat", "Curve '{0}'"), HoveredCurve->GetLongDisplayName()));
 			{
-				MenuBuilder.AddMenuEntry(FCurveEditorCommands::Get().AddKeyHovered);
+				// Buffer Curves
+				MenuBuilder.AddMenuEntry(FCurveEditorCommands::Get().BufferVisibleCurves);
+				MenuBuilder.AddMenuEntry(FCurveEditorCommands::Get().ApplyBufferedCurves, NAME_None, ApplyBufferedCurvesText);
+				MenuBuilder.AddMenuSeparator();
 
-				MenuBuilder.AddMenuEntry(FCurveEditorCommands::Get().ReduceCurve);
-				MenuBuilder.AddMenuEntry(FCurveEditorCommands::Get().BakeCurve);
+				// Modify Curve
+				MenuBuilder.AddMenuEntry(FCurveEditorCommands::Get().AddKeyHovered);
 
 				MenuBuilder.AddSubMenu(LOCTEXT("PreExtrapText", "Pre-Extrap"), FText(), FNewMenuDelegate::CreateLambda(
 					[](FMenuBuilder& SubMenu)
@@ -75,6 +85,12 @@ void FCurveEditorContextMenu::BuildMenu(FMenuBuilder& MenuBuilder, TWeakPtr<FCur
 						SubMenu.AddMenuEntry(FCurveEditorCommands::Get().SetPostInfinityExtrapConstant);
 					})
 				);
+
+				MenuBuilder.AddMenuSeparator();
+
+				// Filters
+				MenuBuilder.AddMenuEntry(FCurveEditorCommands::Get().OpenUserImplementableFilterWindow);
+
 			}
 			MenuBuilder.EndSection();
 		}
@@ -82,8 +98,17 @@ void FCurveEditorContextMenu::BuildMenu(FMenuBuilder& MenuBuilder, TWeakPtr<FCur
 		{
 			MenuBuilder.BeginSection("CurveEditorAllCurveSections", LOCTEXT("CurveEditorAllCurveSections", "All Curves"));
 			{
+				// Buffer Curves
+				MenuBuilder.AddMenuEntry(FCurveEditorCommands::Get().BufferVisibleCurves);
+				MenuBuilder.AddMenuEntry(FCurveEditorCommands::Get().ApplyBufferedCurves, NAME_None, ApplyBufferedCurvesText);
+				MenuBuilder.AddMenuSeparator();
+				
+				// Modify Curves
 				MenuBuilder.AddMenuEntry(FCurveEditorCommands::Get().AddKeyToAllCurves);
-				MenuBuilder.AddMenuEntry(FCurveEditorCommands::Get().AddKeyToAllCurvesHere);
+				MenuBuilder.AddMenuSeparator();
+
+				// Filters
+				MenuBuilder.AddMenuEntry(FCurveEditorCommands::Get().OpenUserImplementableFilterWindow);
 			}
 			MenuBuilder.EndSection();
 		}
diff --git a/Engine/Source/Editor/CurveEditor/Private/CurveEditorContextMenu.h b/Engine/Source/Editor/CurveEditor/Private/CurveEditorContextMenu.h
--- a/Engine/Source/Editor/CurveEditor/Private/CurveEditorContextMenu.h
+++ b/Engine/Source/Editor/CurveEditor/Private/CurveEditorContextMenu.h
@@ -11,5 +11,5 @@ class FMenuBuilder;
 
 struct FCurveEditorContextMenu
 {
-	static void BuildMenu(FMenuBuilder& MenuBuilder, TWeakPtr<FCurveEditor> WeakCurveEditor, TOptional<FCurvePointHandle> ClickedPoint, TOptional<FCurveModelID> HoveredCurve);
+	static void BuildMenu(FMenuBuilder& MenuBuilder, TSharedRef<FCurveEditor> CurveEditor, TOptional<FCurvePointHandle> ClickedPoint, TOptional<FCurveModelID> HoveredCurve);
 };
\ No newline at end of file
diff --git a/Engine/Source/Editor/CurveEditor/Private/CurveEditorModule.cpp b/Engine/Source/Editor/CurveEditor/Private/CurveEditorModule.cpp
--- a/Engine/Source/Editor/CurveEditor/Private/CurveEditorModule.cpp
+++ b/Engine/Source/Editor/CurveEditor/Private/CurveEditorModule.cpp
@@ -1,17 +1,15 @@
 // Copyright 1998-2019 Epic Games, Inc. All Rights Reserved.
 
 #include "Modules/ModuleManager.h"
+#include "ICurveEditorModule.h"
 #include "CurveEditorCommands.h"
+#include "CurveEditor.h"
+#include "CurveEditorViewRegistry.h"
+#include "Framework/MultiBox/MultiBoxExtender.h"
 
-
-class FCurveEditorModule : public IModuleInterface
+class FCurveEditorModule : public ICurveEditorModule
 {
 public:
-	FCurveEditorModule()
-	{
-		check(false);
-	}
-
 	virtual void StartupModule() override
 	{
 		FCurveEditorCommands::Register();
@@ -21,7 +19,69 @@ public:
 	{
 		FCurveEditorCommands::Unregister();
 	}
-};
 
+	virtual FDelegateHandle RegisterEditorExtension(FOnCreateCurveEditorExtension InOnCreateCurveEditorExtension) override
+	{
+		EditorExtensionDelegates.Add(InOnCreateCurveEditorExtension);
+		FDelegateHandle Handle = EditorExtensionDelegates.Last().GetHandle();
+		
+		return Handle;
+	}
+
+	virtual void UnregisterEditorExtension(FDelegateHandle InHandle) override
+	{
+		EditorExtensionDelegates.RemoveAll([=](const FOnCreateCurveEditorExtension& Delegate) { return Delegate.GetHandle() == InHandle; });
+	}
+
+
+	virtual FDelegateHandle RegisterToolExtension(FOnCreateCurveEditorToolExtension InOnCreateCurveEditorToolExtension) override
+	{
+		ToolExtensionDelegates.Add(InOnCreateCurveEditorToolExtension);
+		FDelegateHandle Handle = ToolExtensionDelegates.Last().GetHandle();
+
+		return Handle;
+	}
+
+	virtual void UnregisterToolExtension(FDelegateHandle InHandle) override
+	{
+		ToolExtensionDelegates.RemoveAll([=](const FOnCreateCurveEditorToolExtension& Delegate) { return Delegate.GetHandle() == InHandle; });
+	}
+
+	virtual ECurveEditorViewID RegisterView(FOnCreateCurveEditorView InCreateViewDelegate) override
+	{
+		return FCurveEditorViewRegistry::Get().RegisterCustomView(InCreateViewDelegate);
+	}
+
+	virtual void UnregisterView(ECurveEditorViewID InViewID) override
+	{
+		return FCurveEditorViewRegistry::Get().UnregisterCustomView(InViewID);
+	}
+
+	virtual TArray<FCurveEditorMenuExtender>& GetAllToolBarMenuExtenders() override
+	{
+		return ToolBarMenuExtenders;
+	}
+
+	virtual TArrayView<const FOnCreateCurveEditorExtension> GetEditorExtensions() const override
+	{
+		return EditorExtensionDelegates;
+	}
+
+	virtual TArrayView<const FOnCreateCurveEditorToolExtension> GetToolExtensions() const override
+	{
+		return ToolExtensionDelegates;
+	}
+
+private:
+	/** List of editor extension handler delegates Curve Editors will execute when they are created. */
+	TArray<FOnCreateCurveEditorExtension> EditorExtensionDelegates;
+
+	/** List of tool extension handler delegates Curve Editors will execute when they are created. */
+	TArray<FOnCreateCurveEditorToolExtension> ToolExtensionDelegates;
+
+	/** List of Extenders that that should be called when building the Curve Editor toolbar. */
+	TArray<FCurveEditorMenuExtender> ToolBarMenuExtenders;
+
+};
 
 IMPLEMENT_MODULE(FCurveEditorModule, CurveEditor)
\ No newline at end of file
diff --git a/Engine/Source/Editor/CurveEditor/Private/CurveEditorSettings.cpp b/Engine/Source/Editor/CurveEditor/Private/CurveEditorSettings.cpp
--- a/Engine/Source/Editor/CurveEditor/Private/CurveEditorSettings.cpp
+++ b/Engine/Source/Editor/CurveEditor/Private/CurveEditorSettings.cpp
@@ -2,12 +2,12 @@
 
 #include "CurveEditorSettings.h"
 
-UCurveEditorSettings::UCurveEditorSettings( const FObjectInitializer& ObjectInitializer )
-	: Super( ObjectInitializer )
+UCurveEditorSettings::UCurveEditorSettings()
 {
 	bAutoFrameCurveEditor = true;
 	bShowCurveEditorCurveToolTips = true;
 	TangentVisibility = ECurveEditorTangentVisibility::SelectedKeys;
+	ZoomPosition = ECurveEditorZoomPosition::CurrentTime;
 }
 
 bool UCurveEditorSettings::GetAutoFrameCurveEditor() const
@@ -38,16 +38,31 @@ void UCurveEditorSettings::SetShowCurveEditorCurveToolTips(bool InbShowCurveEdit
 	}
 }
 
-ECurveEditorTangentVisibility::Type UCurveEditorSettings::GetTangentVisibility() const
+ECurveEditorTangentVisibility UCurveEditorSettings::GetTangentVisibility() const
 {
 	return TangentVisibility;
 }
 
-void UCurveEditorSettings::SetTangentVisibility(ECurveEditorTangentVisibility::Type InTangentVisibility)
+void UCurveEditorSettings::SetTangentVisibility(ECurveEditorTangentVisibility InTangentVisibility)
 {
 	if (TangentVisibility != InTangentVisibility)
 	{
 		TangentVisibility = InTangentVisibility;
 		SaveConfig();
 	}
-}
\ No newline at end of file
+}
+
+ECurveEditorZoomPosition UCurveEditorSettings::GetZoomPosition() const
+{
+	return ZoomPosition;
+}
+
+void UCurveEditorSettings::SetZoomPosition(ECurveEditorZoomPosition InZoomPosition)
+{
+	if (ZoomPosition != InZoomPosition)
+	{
+		ZoomPosition = InZoomPosition;
+		SaveConfig();
+	}
+}
+
diff --git a/Engine/Source/Editor/CurveEditor/Private/DragOperations/CurveEditorDragOperation_Marquee.cpp b/Engine/Source/Editor/CurveEditor/Private/DragOperations/CurveEditorDragOperation_Marquee.cpp
--- a/Engine/Source/Editor/CurveEditor/Private/DragOperations/CurveEditorDragOperation_Marquee.cpp
+++ b/Engine/Source/Editor/CurveEditor/Private/DragOperations/CurveEditorDragOperation_Marquee.cpp
@@ -2,14 +2,21 @@
 
 #include "DragOperations/CurveEditorDragOperation_Marquee.h"
 #include "CurveEditor.h"
+#include "SCurveEditorView.h"
 #include "SCurveEditorPanel.h"
 #include "EditorStyleSet.h"
 #include "CurveDrawInfo.h"
 
-FCurveEditorDragOperation_Marquee::FCurveEditorDragOperation_Marquee(FCurveEditor* InCurveEditor, const SCurveEditorPanel* InCurveEditorPanel)
+FCurveEditorDragOperation_Marquee::FCurveEditorDragOperation_Marquee(FCurveEditor* InCurveEditor)
+	: LockedToView(nullptr)
+{
+	CurveEditor = InCurveEditor;
+}
+
+FCurveEditorDragOperation_Marquee::FCurveEditorDragOperation_Marquee(FCurveEditor* InCurveEditor, SCurveEditorView* InLockedToView)
+	: LockedToView(InLockedToView)
 {
 	CurveEditor = InCurveEditor;
-	CurveEditorPanel = InCurveEditorPanel;
 }
 
 void FCurveEditorDragOperation_Marquee::OnBeginDrag(FVector2D InitialPosition, FVector2D CurrentPosition, const FPointerEvent& MouseEvent)
@@ -41,14 +48,39 @@ void FCurveEditorDragOperation_Marquee::OnEndDrag(FVector2D InitialPosition, FVe
 		FMath::Max(InitialPosition.Y, CurrentPosition.Y)
 		);
 
-	FCurveEditorScreenSpace ScreenSpace = CurveEditor->GetScreenSpace();
+	TArray<FCurvePointHandle> AllPoints;
+
+	if (LockedToView)
+	{
+		LockedToView->GetPointsWithinWidgetRange(Marquee, &AllPoints);
+	}
+	else
+	{
+		TSharedPtr<SCurveEditorPanel> CurveEditorPanel = CurveEditor->GetPanel();
+
+		FGeometry ViewContainerGeometry = CurveEditorPanel->GetViewContainerGeometry();
+		FSlateLayoutTransform InverseContainerTransform = ViewContainerGeometry.GetAccumulatedLayoutTransform().Inverse();
+		for (TSharedPtr<SCurveEditorView> View : CurveEditorPanel->GetViews())
+		{
+			const FGeometry& LocalGeometry = View->GetCachedGeometry();
+			FSlateLayoutTransform ContainerToView = InverseContainerTransform.Concatenate(LocalGeometry.GetAccumulatedLayoutTransform()).Inverse();
+
+			FSlateRect UnclippedLocalMarquee = FSlateRect(ContainerToView.TransformPoint(Marquee.GetTopLeft()), ContainerToView.TransformPoint(Marquee.GetBottomRight()));
+			FSlateRect ClippedLocalMarquee = UnclippedLocalMarquee.IntersectionWith(FSlateRect(FVector2D(0.f,0.f), LocalGeometry.GetLocalSize()));
+
+			if (ClippedLocalMarquee.IsValid() && !ClippedLocalMarquee.IsEmpty())
+			{
+				View->GetPointsWithinWidgetRange(ClippedLocalMarquee, &AllPoints);
+			}
+		}
+	}
 
-	bool bRemoveFromSelection = MouseEvent.IsAltDown();
+	const bool bIsShiftDown = MouseEvent.IsShiftDown();
+	const bool bRemoveFromSelection = MouseEvent.IsAltDown();
 
-	// Only select the same types of point
 	TOptional<ECurvePointType> MatchPointType;
 
-	if (!MouseEvent.IsShiftDown() && !bRemoveFromSelection)
+	if (!bIsShiftDown && !bRemoveFromSelection)
 	{
 		CurveEditor->Selection.Clear();
 	}
@@ -57,49 +89,39 @@ void FCurveEditorDragOperation_Marquee::OnEndDrag(FVector2D InitialPosition, FVe
 		MatchPointType = CurveEditor->Selection.GetSelectionType();
 	}
 
-	FSlateRect MarqueeRectPx = Marquee;
-	for (const FCurveDrawParams& DrawParams : CurveEditorPanel->GetCachedDrawParams())
+	// Now that we've gathered the overlapping points, perform the relevant selection
+	for (const FCurvePointHandle& Point : AllPoints)
 	{
-		for (const FCurvePointInfo& Point : DrawParams.Points)
+		// If we're able to specify a point type, prefer keys
+		if (!MatchPointType.IsSet() && Point.PointType == ECurvePointType::Key)
 		{
-			// Can we select this type of point?
-			if (MatchPointType.IsSet() && Point.Type != MatchPointType.GetValue())
-			{
-				continue;
-			}
-
-			const FKeyDrawInfo& DrawInfo = DrawParams.GetKeyDrawInfo(Point.Type);
-			FSlateRect PointRect = FSlateRect::FromPointAndExtent(Point.ScreenPosition - DrawInfo.ScreenSize/2, DrawInfo.ScreenSize);
+			// We found a key, so start selecting with keys
+			CurveEditor->Selection.Clear();
+			MatchPointType.Emplace(ECurvePointType::Key);
+		}
 
-			if (FSlateRect::DoRectanglesIntersect(PointRect, MarqueeRectPx))
-			{
-				if (!MatchPointType.IsSet() && Point.Type == ECurvePointType::Key)
-				{
-					CurveEditor->Selection.Clear();
-					MatchPointType = ECurvePointType::Key;
-				}
-
-				if (bRemoveFromSelection)
-				{
-					CurveEditor->Selection.Remove(DrawParams.GetID(), Point.Type, Point.KeyHandle);
-				}
-				else
-				{
-					CurveEditor->Selection.Add(DrawParams.GetID(), Point.Type, Point.KeyHandle);
-				}
-			}
+		if (MatchPointType.IsSet() && Point.PointType != MatchPointType.GetValue())
+		{
+			// Point does not match what we want
+			continue;
+		}
+		else if (bRemoveFromSelection)
+		{
+			CurveEditor->Selection.Remove(Point);
+		}
+		else
+		{
+			CurveEditor->Selection.Add(Point);
 		}
 	}
 }
 
-int32 FCurveEditorDragOperation_Marquee::OnPaint(const FGeometry& AllottedGeometry, FSlateWindowElementList& OutDrawElements, int32 LayerId)
+void FCurveEditorDragOperation_Marquee::OnPaint(const FGeometry& AllottedGeometry, FSlateWindowElementList& OutDrawElements, int32 PaintOnLayerId)
 {
 	FSlateDrawElement::MakeBox(
 		OutDrawElements,
-		LayerId,
+		PaintOnLayerId,
 		AllottedGeometry.ToPaintGeometry(Marquee.GetTopLeft(), Marquee.GetBottomRight() - Marquee.GetTopLeft()),
 		FEditorStyle::GetBrush(TEXT("MarqueeSelection"))
 		);
-
-	return LayerId;
 }
diff --git a/Engine/Source/Editor/CurveEditor/Private/DragOperations/CurveEditorDragOperation_Marquee.h b/Engine/Source/Editor/CurveEditor/Private/DragOperations/CurveEditorDragOperation_Marquee.h
--- a/Engine/Source/Editor/CurveEditor/Private/DragOperations/CurveEditorDragOperation_Marquee.h
+++ b/Engine/Source/Editor/CurveEditor/Private/DragOperations/CurveEditorDragOperation_Marquee.h
@@ -10,18 +10,20 @@
 #include "ICurveEditorDragOperation.h"
 
 class FCurveEditor;
+class SCurveEditorView;
 class SCurveEditorPanel;
 
 class FCurveEditorDragOperation_Marquee : public ICurveEditorDragOperation
 {
 public:
 
-	FCurveEditorDragOperation_Marquee(FCurveEditor* InCurveEditor, const SCurveEditorPanel* InCurveEditorPanel);
+	FCurveEditorDragOperation_Marquee(FCurveEditor* InCurveEditor);
+	FCurveEditorDragOperation_Marquee(FCurveEditor* InCurveEditor, SCurveEditorView*  InLockedToView);
 
 	virtual void OnBeginDrag(FVector2D InitialPosition, FVector2D CurrentPosition, const FPointerEvent& MouseEvent) override;
 	virtual void OnDrag(FVector2D InitialPosition, FVector2D CurrentPosition, const FPointerEvent& MouseEvent) override;
 	virtual void OnEndDrag(FVector2D InitialPosition, FVector2D CurrentPosition, const FPointerEvent& MouseEvent) override;
-	virtual int32 OnPaint(const FGeometry& AllottedGeometry, FSlateWindowElementList& OutDrawElements, int32 LayerId) override;
+	virtual void OnPaint(const FGeometry& AllottedGeometry, FSlateWindowElementList& OutDrawElements, int32 PaintOnLayerId) override;
 
 private:
 
@@ -29,6 +31,6 @@ private:
 	FSlateRect Marquee;
 	/** Ptr back to the curve editor */
 	FCurveEditor* CurveEditor;
-	/** Ptr back to the curve editor panel */
-	const SCurveEditorPanel* CurveEditorPanel;
+	/** When valid, marquee selection should only occur inside this view; all geometries are in local space */
+	SCurveEditorView* LockedToView;
 };
\ No newline at end of file
diff --git a/Engine/Source/Editor/CurveEditor/Private/DragOperations/CurveEditorDragOperation_MoveKeys.cpp b/Engine/Source/Editor/CurveEditor/Private/DragOperations/CurveEditorDragOperation_MoveKeys.cpp
--- a/Engine/Source/Editor/CurveEditor/Private/DragOperations/CurveEditorDragOperation_MoveKeys.cpp
+++ b/Engine/Source/Editor/CurveEditor/Private/DragOperations/CurveEditorDragOperation_MoveKeys.cpp
@@ -3,26 +3,27 @@
 #include "DragOperations/CurveEditorDragOperation_MoveKeys.h"
 #include "CurveEditorScreenSpace.h"
 #include "CurveEditor.h"
+#include "SCurveEditorView.h"
 
 void FCurveEditorDragOperation_MoveKeys::OnInitialize(FCurveEditor* InCurveEditor, const TOptional<FCurvePointHandle>& CardinalPoint)
 {
 	CurveEditor = InCurveEditor;
 }
 
-
 void FCurveEditorDragOperation_MoveKeys::OnBeginDrag(FVector2D InitialPosition, FVector2D CurrentPosition, const FPointerEvent& MouseEvent)
 {
-	int32 NumKeys = CurveEditor->Selection.Count();
+	int32 NumKeys = CurveEditor->GetSelection().Count();
 	Transaction = MakeUnique<FScopedTransaction>(FText::Format(NSLOCTEXT("CurveEditor", "MoveKeysFormat", "Move {0}|plural(one=Key, other=Keys)"), NumKeys));
 
 	KeysByCurve.Reset();
+	CurveEditor->SuppressBoundTransformUpdates(true);
 
-	for (const TTuple<FCurveModelID, FKeyHandleSet>& Pair : CurveEditor->Selection.GetAll())
+	for (const TTuple<FCurveModelID, FKeyHandleSet>& Pair : CurveEditor->GetSelection().GetAll())
 	{
 		FCurveModelID CurveID = Pair.Key;
 		FCurveModel*  Curve   = CurveEditor->FindCurve(CurveID);
 
-		if (ensure(Curve))
+		if (ensureAlways(Curve))
 		{
 			Curve->Modify();
 
@@ -40,35 +41,42 @@ void FCurveEditorDragOperation_MoveKeys::OnBeginDrag(FVector2D InitialPosition,
 void FCurveEditorDragOperation_MoveKeys::OnDrag(FVector2D InitialPosition, FVector2D CurrentPosition, const FPointerEvent& MouseEvent)
 {
 	TArray<FKeyPosition> NewKeyPositionScratch;
-
-	FVector2D MousePosition = GetLockedMousePosition(InitialPosition, CurrentPosition, MouseEvent);
-
-	FCurveEditorScreenSpace ScreenSpace = CurveEditor->GetScreenSpace();
-	double DeltaInput  =  (MousePosition.X - InitialPosition.X) / ScreenSpace.PixelsPerInput();
-	double DeltaOutput = -(MousePosition.Y - InitialPosition.Y) / ScreenSpace.PixelsPerOutput();
+	FVector2D MousePosition = CurveEditor->GetAxisSnap().GetSnappedPosition(InitialPosition, CurrentPosition, MouseEvent);
 
 	for (const FKeyData& KeyData : KeysByCurve)
 	{
-		FCurveModel* Curve = CurveEditor->FindCurve(KeyData.CurveID);
+		const SCurveEditorView* View = CurveEditor->FindFirstInteractiveView(KeyData.CurveID);
+		if (!View)
+		{
+			continue;
+		}
 
-		if (ensure(Curve))
+		FCurveModel* Curve = CurveEditor->FindCurve(KeyData.CurveID);
+		if (!ensureAlways(Curve))
 		{
-			NewKeyPositionScratch.Reset();
-			NewKeyPositionScratch.Reserve(KeyData.StartKeyPositions.Num());
+			continue;
+		}
 
-			for (FKeyPosition StartPosition : KeyData.StartKeyPositions)
-			{
-				StartPosition.InputValue  += DeltaInput;
-				StartPosition.OutputValue += DeltaOutput;
+		FCurveEditorScreenSpace CurveSpace = View->GetCurveSpace(KeyData.CurveID);
 
-				StartPosition.InputValue  = SnapMetrics.SnapInputSeconds(StartPosition.InputValue);
-				StartPosition.OutputValue = SnapMetrics.SnapOutput(StartPosition.OutputValue);
+		double DeltaInput = (MousePosition.X - InitialPosition.X) / CurveSpace.PixelsPerInput();
+		double DeltaOutput = -(MousePosition.Y - InitialPosition.Y) / CurveSpace.PixelsPerOutput();
 
-				NewKeyPositionScratch.Add(StartPosition);
-			}
+		NewKeyPositionScratch.Reset();
+		NewKeyPositionScratch.Reserve(KeyData.StartKeyPositions.Num());
 
-			Curve->SetKeyPositions(KeyData.Handles, NewKeyPositionScratch);
+		for (FKeyPosition StartPosition : KeyData.StartKeyPositions)
+		{
+			StartPosition.InputValue  += DeltaInput;
+			StartPosition.OutputValue += DeltaOutput;
+
+			StartPosition.InputValue  = View->IsTimeSnapEnabled() ? SnapMetrics.SnapInputSeconds(StartPosition.InputValue) : StartPosition.InputValue;
+			StartPosition.OutputValue = View->IsValueSnapEnabled() ? SnapMetrics.SnapOutput(StartPosition.OutputValue) : StartPosition.OutputValue;
+
+			NewKeyPositionScratch.Add(StartPosition);
 		}
+
+		Curve->SetKeyPositions(KeyData.Handles, NewKeyPositionScratch);
 	}
 }
 
@@ -83,4 +91,12 @@ void FCurveEditorDragOperation_MoveKeys::OnCancelDrag()
 			Curve->SetKeyPositions(KeyData.Handles, KeyData.StartKeyPositions);
 		}
 	}
-}
\ No newline at end of file
+
+	CurveEditor->SuppressBoundTransformUpdates(false);
+}
+
+void FCurveEditorDragOperation_MoveKeys::OnEndDrag(FVector2D InitialPosition, FVector2D CurrentPosition, const FPointerEvent& MouseEvent)
+{
+	ICurveEditorKeyDragOperation::OnEndDrag(InitialPosition, CurrentPosition, MouseEvent);
+	CurveEditor->SuppressBoundTransformUpdates(false);
+}
diff --git a/Engine/Source/Editor/CurveEditor/Private/DragOperations/CurveEditorDragOperation_MoveKeys.h b/Engine/Source/Editor/CurveEditor/Private/DragOperations/CurveEditorDragOperation_MoveKeys.h
--- a/Engine/Source/Editor/CurveEditor/Private/DragOperations/CurveEditorDragOperation_MoveKeys.h
+++ b/Engine/Source/Editor/CurveEditor/Private/DragOperations/CurveEditorDragOperation_MoveKeys.h
@@ -17,6 +17,7 @@ public:
 	virtual void OnInitialize(FCurveEditor* CurveEditor, const TOptional<FCurvePointHandle>& CardinalPoint) override;
 	virtual void OnBeginDrag(FVector2D InitialPosition, FVector2D CurrentPosition, const FPointerEvent& MouseEvent) override;
 	virtual void OnDrag(FVector2D InitialPosition, FVector2D CurrentPosition, const FPointerEvent& MouseEvent) override;
+	virtual void OnEndDrag(FVector2D InitialPosition, FVector2D CurrentPosition, const FPointerEvent& MouseEvent) override;
 	virtual void OnCancelDrag() override;
 
 private:
diff --git a/Engine/Source/Editor/CurveEditor/Private/DragOperations/CurveEditorDragOperation_Pan.cpp b/Engine/Source/Editor/CurveEditor/Private/DragOperations/CurveEditorDragOperation_Pan.cpp
--- a/Engine/Source/Editor/CurveEditor/Private/DragOperations/CurveEditorDragOperation_Pan.cpp
+++ b/Engine/Source/Editor/CurveEditor/Private/DragOperations/CurveEditorDragOperation_Pan.cpp
@@ -3,34 +3,61 @@
 #include "CurveEditorDragOperation_Pan.h"
 #include "CurveEditorScreenSpace.h"
 #include "CurveEditor.h"
+#include "SCurveEditorView.h"
+#include "SCurveEditorPanel.h"
 
-FCurveEditorDragOperation_Pan::FCurveEditorDragOperation_Pan(FCurveEditor* InCurveEditor)
+FCurveEditorDragOperation_PanView::FCurveEditorDragOperation_PanView(FCurveEditor* InCurveEditor, TSharedPtr<SCurveEditorView> InView)
 	: CurveEditor(InCurveEditor)
+	, View(InView)
+{}
+
+void FCurveEditorDragOperation_PanView::OnBeginDrag(FVector2D InitialPosition, FVector2D CurrentPosition, const FPointerEvent& MouseEvent)
 {
+	FCurveEditorScreenSpace ViewSpace = View->GetViewSpace();
+
+	InitialInputMin = ViewSpace.GetInputMin();
+	InitialInputMax = ViewSpace.GetInputMax();
+	InitialOutputMin = ViewSpace.GetOutputMin();
+	InitialOutputMax = ViewSpace.GetOutputMax();
 }
 
-void FCurveEditorDragOperation_Pan::OnBeginDrag(FVector2D InitialPosition, FVector2D CurrentPosition, const FPointerEvent& MouseEvent)
+void FCurveEditorDragOperation_PanView::OnDrag(FVector2D InitialPosition, FVector2D CurrentPosition, const FPointerEvent& MouseEvent)
 {
-	FCurveEditorScreenSpace ScreenSpace = CurveEditor->GetScreenSpace();
+	FVector2D PixelDelta = CurveEditor->GetAxisSnap().GetSnappedPosition(InitialPosition, CurrentPosition, MouseEvent, true) - InitialPosition;
+
+	FCurveEditorScreenSpace ViewSpace = View->GetViewSpace();
 
-	InitialInputMin  = ScreenSpace.GetInputMin();
-	InitialInputMax  = ScreenSpace.GetInputMax();
-	InitialOutputMin = ScreenSpace.GetOutputMin();
-	InitialOutputMax = ScreenSpace.GetOutputMax();
+	double InputMin = InitialInputMin - PixelDelta.X / ViewSpace.PixelsPerInput();
+	double InputMax = InitialInputMax - PixelDelta.X / ViewSpace.PixelsPerInput();
+
+	double OutputMin = InitialOutputMin + PixelDelta.Y / ViewSpace.PixelsPerOutput();
+	double OutputMax = InitialOutputMax + PixelDelta.Y / ViewSpace.PixelsPerOutput();
+
+	CurveEditor->GetBounds().SetInputBounds(InputMin, InputMax);
+	View->SetOutputBounds(OutputMin, OutputMax);
 }
 
-void FCurveEditorDragOperation_Pan::OnDrag(FVector2D InitialPosition, FVector2D CurrentPosition, const FPointerEvent& MouseEvent)
+FCurveEditorDragOperation_PanInput::FCurveEditorDragOperation_PanInput(FCurveEditor* InCurveEditor)
+	: CurveEditor(InCurveEditor)
+{}
+
+void FCurveEditorDragOperation_PanInput::OnBeginDrag(FVector2D InitialPosition, FVector2D CurrentPosition, const FPointerEvent& MouseEvent)
 {
-	FVector2D PixelDelta = GetLockedMousePosition(InitialPosition, CurrentPosition, MouseEvent) - InitialPosition;
+	FCurveEditorScreenSpaceH InputSpace = CurveEditor->GetPanelInputSpace();
+	InitialInputMin = InputSpace.GetInputMin();
+	InitialInputMax = InputSpace.GetInputMax();
+}
 
-	FCurveEditorScreenSpace ScreenSpace = CurveEditor->GetScreenSpace();
+void FCurveEditorDragOperation_PanInput::OnDrag(FVector2D InitialPosition, FVector2D CurrentPosition, const FPointerEvent& MouseEvent)
+{
+	FVector2D PixelDelta = CurveEditor->GetAxisSnap().GetSnappedPosition(InitialPosition, CurrentPosition, MouseEvent, true) - InitialPosition;
 
-	double InputMin  = InitialInputMin - PixelDelta.X / ScreenSpace.PixelsPerInput();
-	double InputMax  = InitialInputMax - PixelDelta.X / ScreenSpace.PixelsPerInput();
+	FCurveEditorScreenSpaceH InputSpace = CurveEditor->GetPanelInputSpace();
 
-	double OutputMin = InitialOutputMin + PixelDelta.Y / ScreenSpace.PixelsPerOutput();
-	double OutputMax = InitialOutputMax + PixelDelta.Y / ScreenSpace.PixelsPerOutput();
+	double InputMin = InitialInputMin - PixelDelta.X / InputSpace.PixelsPerInput();
+	double InputMax = InitialInputMax - PixelDelta.X / InputSpace.PixelsPerInput();
 
 	CurveEditor->GetBounds().SetInputBounds(InputMin, InputMax);
-	CurveEditor->GetBounds().SetOutputBounds(OutputMin, OutputMax);
+
+	CurveEditor->GetPanel()->ScrollBy(-MouseEvent.GetCursorDelta().Y);
 }
\ No newline at end of file
diff --git a/Engine/Source/Editor/CurveEditor/Private/DragOperations/CurveEditorDragOperation_Pan.h b/Engine/Source/Editor/CurveEditor/Private/DragOperations/CurveEditorDragOperation_Pan.h
--- a/Engine/Source/Editor/CurveEditor/Private/DragOperations/CurveEditorDragOperation_Pan.h
+++ b/Engine/Source/Editor/CurveEditor/Private/DragOperations/CurveEditorDragOperation_Pan.h
@@ -5,11 +5,12 @@
 #include "ICurveEditorDragOperation.h"
 
 class FCurveEditor;
+class SCurveEditorView;
 
-class FCurveEditorDragOperation_Pan : public ICurveEditorDragOperation
+class FCurveEditorDragOperation_PanView : public ICurveEditorDragOperation
 {
 public:
-	FCurveEditorDragOperation_Pan(FCurveEditor* CurveEditor);
+	FCurveEditorDragOperation_PanView(FCurveEditor* CurveEditor, TSharedPtr<SCurveEditorView> InView);
 
 	virtual void OnBeginDrag(FVector2D InitialPosition, FVector2D CurrentPosition, const FPointerEvent& MouseEvent) override;
 	virtual void OnDrag(FVector2D InitialPosition, FVector2D CurrentPosition, const FPointerEvent& MouseEvent) override;
@@ -17,6 +18,22 @@ public:
 private:
 
 	FCurveEditor* CurveEditor;
+	TSharedPtr<SCurveEditorView> View;
+
 	double InitialInputMin, InitialInputMax;
 	double InitialOutputMin, InitialOutputMax;
+};
+
+class FCurveEditorDragOperation_PanInput : public ICurveEditorDragOperation
+{
+public:
+	FCurveEditorDragOperation_PanInput(FCurveEditor* CurveEditor);
+
+	virtual void OnBeginDrag(FVector2D InitialPosition, FVector2D CurrentPosition, const FPointerEvent& MouseEvent) override;
+	virtual void OnDrag(FVector2D InitialPosition, FVector2D CurrentPosition, const FPointerEvent& MouseEvent) override;
+
+private:
+
+	FCurveEditor* CurveEditor;
+	double InitialInputMin, InitialInputMax;
 };
\ No newline at end of file
