This is a significant overhaul to the Curve Editor used by Sequencer which adds a plugin-based architecture and extensibility. New tools and toolbar buttons can be added to all usages of the curve edi

diff --git a/Engine/Source/Editor/CurveEditor/Private/DragOperations/CurveEditorDragOperation_Tangent.cpp b/Engine/Source/Editor/CurveEditor/Private/DragOperations/CurveEditorDragOperation_Tangent.cpp
--- a/Engine/Source/Editor/CurveEditor/Private/DragOperations/CurveEditorDragOperation_Tangent.cpp
+++ b/Engine/Source/Editor/CurveEditor/Private/DragOperations/CurveEditorDragOperation_Tangent.cpp
@@ -3,6 +3,14 @@
 #include "CurveEditorDragOperation_Tangent.h"
 #include "CurveEditorScreenSpace.h"
 #include "CurveEditor.h"
+#include "CurveEditorHelpers.h"
+#include "SCurveEditorView.h"
+
+namespace CurveEditorDragOperation
+{
+	/** If they're within this many pixels of crossing over the straight up/down boundary, we clamp the handles. */
+	constexpr float TangentCrossoverThresholdPx = 1.f;
+}
 
 void FCurveEditorDragOperation_Tangent::OnInitialize(FCurveEditor* InCurveEditor, const TOptional<FCurvePointHandle>& InCardinalPoint)
 {
@@ -11,23 +19,24 @@ void FCurveEditorDragOperation_Tangent::OnInitialize(FCurveEditor* InCurveEditor
 
 void FCurveEditorDragOperation_Tangent::OnBeginDrag(FVector2D InitialPosition, FVector2D CurrentPosition, const FPointerEvent& MouseEvent)
 {
-	PointType = CurveEditor->Selection.GetSelectionType();
+	PointType = CurveEditor->GetSelection().GetSelectionType();
 
-	int32 NumKeys = CurveEditor->Selection.Count();
+	int32 NumKeys = CurveEditor->GetSelection().Count();
 
 	FText Description = PointType == ECurvePointType::ArriveTangent
 		? FText::Format(NSLOCTEXT("CurveEditor", "DragEntryTangentsFormat", "Drag Entry {0}|plural(one=Tangent, other=Tangents)"), NumKeys)
 		: FText::Format(NSLOCTEXT("CurveEditor", "DragExitTangentsFormat", "Drag Exit {0}|plural(one=Tangent, other=Tangents)"), NumKeys);
 
 	Transaction = MakeUnique<FScopedTransaction>(Description);
+	CurveEditor->SuppressBoundTransformUpdates(true);
 
 	KeysByCurve.Reset();
-	for (const TTuple<FCurveModelID, FKeyHandleSet>& Pair : CurveEditor->Selection.GetAll())
+	for (const TTuple<FCurveModelID, FKeyHandleSet>& Pair : CurveEditor->GetSelection().GetAll())
 	{
 		FCurveModelID CurveID = Pair.Key;
 		FCurveModel* Curve    = CurveEditor->FindCurve(CurveID);
 
-		if (ensure(Curve))
+		if (ensureAlways(Curve))
 		{
 			Curve->Modify();
 
@@ -44,116 +53,151 @@ void FCurveEditorDragOperation_Tangent::OnBeginDrag(FVector2D InitialPosition, F
 
 void FCurveEditorDragOperation_Tangent::OnDrag(FVector2D InitialPosition, FVector2D CurrentPosition, const FPointerEvent& MouseEvent)
 {
-	FCurveEditorScreenSpace ScreenSpace = CurveEditor->GetScreenSpace();
-	const float DisplayRatio = (ScreenSpace.PixelsPerOutput() / ScreenSpace.PixelsPerInput());
-	
 	FVector2D PixelDelta = CurrentPosition - InitialPosition;
-	if (MouseEvent.IsShiftDown())
-	{
-		PixelDelta = RoundTrajectory(PixelDelta);
-	}
 
 	TArray<FKeyAttributes> NewKeyAttributesScratch;
 	TArray<FKeyPosition> KeyPositions;
 
 	for (const FKeyData& KeyData : KeysByCurve)
 	{
+		const SCurveEditorView* View = CurveEditor->FindFirstInteractiveView(KeyData.CurveID);
+		if (!View)
+		{
+			continue;
+		}
+
 		FCurveModel* Curve = CurveEditor->FindCurve(KeyData.CurveID);
-		if (ensure(Curve))
+		if (!ensureAlways(Curve))
 		{
-			NewKeyAttributesScratch.Reset();
-			NewKeyAttributesScratch.Reserve(KeyData.Attributes.Num());
-			KeyPositions.SetNumUninitialized(KeyData.Attributes.Num());
-			Curve->GetKeyPositions(KeyData.Handles, KeyPositions);
-			
+			continue;
+		}
 
-			if (PointType == ECurvePointType::ArriveTangent)
+		FCurveEditorScreenSpace CurveSpace = View->GetCurveSpace(KeyData.CurveID);
+		const float DisplayRatio = (CurveSpace.PixelsPerOutput() / CurveSpace.PixelsPerInput());
+
+		NewKeyAttributesScratch.Reset();
+		NewKeyAttributesScratch.Reserve(KeyData.Attributes.Num());
+		KeyPositions.SetNumUninitialized(KeyData.Attributes.Num());
+		Curve->GetKeyPositions(KeyData.Handles, KeyPositions);
+
+		if (PointType == ECurvePointType::ArriveTangent)
+		{
+			for (int32 i = 0; i < KeyData.Handles.Num(); ++i)
 			{
-				for (int32 i = 0; i < KeyData.Handles.Num(); ++i)
+				const FKeyAttributes Attributes = KeyData.Attributes[i];
+				if (Attributes.HasArriveTangent())
 				{
-					const FKeyAttributes Attributes = KeyData.Attributes[i];
-					if (Attributes.HasArriveTangent())
+					const float ArriveTangent = Attributes.GetArriveTangent();
+
+					if (Attributes.HasTangentWeightMode() && Attributes.HasArriveTangentWeight() &&
+						(Attributes.GetTangentWeightMode() == RCTWM_WeightedBoth || Attributes.GetTangentWeightMode() == RCTWM_WeightedArrive))
 					{
-						const float ArriveTangent = Attributes.GetArriveTangent();
-						const FVector2D  KeyScreenPosition = FVector2D(ScreenSpace.SecondsToScreen(KeyPositions[i].InputValue), ScreenSpace.ValueToScreen(KeyPositions[i].OutputValue));
-						if (Attributes.HasTangentWeightMode() && Attributes.HasArriveTangentWeight() &&
-							(Attributes.GetTangentWeightMode() == RCTWM_WeightedBoth || Attributes.GetTangentWeightMode() == RCTWM_WeightedArrive))
+						FVector2D TangentOffset = CurveEditor::ComputeScreenSpaceTangentOffset(CurveSpace, ArriveTangent, -Attributes.GetArriveTangentWeight());
+						TangentOffset += PixelDelta;
+
+						if (MouseEvent.IsShiftDown())
 						{
-							float Tangent, Weight;
-							FVector2D TangentScreenPosition = CurveEditor->GetTangentPositionInScreenSpace(KeyScreenPosition, ArriveTangent, -Attributes.GetArriveTangentWeight());
-							TangentScreenPosition.X += PixelDelta.X;
-							TangentScreenPosition.Y += PixelDelta.Y;
+							TangentOffset = RoundTrajectory(TangentOffset);
+						}
 
-							CurveEditor->GetTangentAndWeightFromScreenPosition(KeyScreenPosition, TangentScreenPosition, Tangent, Weight);
+						// We prevent the handle from crossing over the 0 point... the code handles it but it creates an ugly pop in your curve and it lets your Arrive tangents become Leave tangents which defeats the point.
+						TangentOffset.X = FMath::Min(TangentOffset.X, -CurveEditorDragOperation::TangentCrossoverThresholdPx);
 
-							FKeyAttributes NewAttributes;
+						float Tangent, Weight;
+						CurveEditor::TangentAndWeightFromOffset(CurveSpace, TangentOffset, Tangent, Weight);
 
-							NewAttributes.SetArriveTangent(Tangent);
-							NewAttributes.SetArriveTangentWeight(Weight);
-							NewKeyAttributesScratch.Add(NewAttributes);
+						FKeyAttributes NewAttributes;
 
-						}
-						else
-						{
-							const float PixelLength = 60.0f;
-							FVector2D TangentScreenPosition = FCurveEditor::GetVectorFromSlopeAndLength(ArriveTangent * -DisplayRatio, -PixelLength) + KeyScreenPosition;
-							TangentScreenPosition += PixelDelta;
-							TangentScreenPosition -= KeyScreenPosition;
-							const float Tangent = (-TangentScreenPosition.Y / TangentScreenPosition.X) / DisplayRatio;
-							FKeyAttributes NewAttributes;
-							NewAttributes.SetArriveTangent(Tangent);
-							NewKeyAttributesScratch.Add(NewAttributes);
-						}
+						NewAttributes.SetArriveTangent(Tangent);
+						NewAttributes.SetArriveTangentWeight(Weight);
+						NewKeyAttributesScratch.Add(NewAttributes);
 					}
-					else //still need to add since expect attributes to equal num of selected
+					else
 					{
-						NewKeyAttributesScratch.Add(FKeyAttributes());
+						const float PixelLength = 60.0f;
+						FVector2D TangentOffset = CurveEditor::GetVectorFromSlopeAndLength(ArriveTangent * -DisplayRatio, -PixelLength);
+						TangentOffset += PixelDelta;
+
+						if (MouseEvent.IsShiftDown())
+						{
+							TangentOffset = RoundTrajectory(TangentOffset);
+						}
+
+						// We prevent the handle from crossing over the 0 point... the code handles it but it creates an ugly pop in your curve and it lets your Arrive tangents become Leave tangents which defeats the point.
+						TangentOffset.X = FMath::Min(TangentOffset.X, -CurveEditorDragOperation::TangentCrossoverThresholdPx);
+
+						const float Tangent = (-TangentOffset.Y / TangentOffset.X) / DisplayRatio;
+
+						FKeyAttributes NewAttributes;
+						NewAttributes.SetArriveTangent(Tangent);
+						NewKeyAttributesScratch.Add(NewAttributes);
 					}
 				}
+				else //still need to add since expect attributes to equal num of selected
+				{
+					NewKeyAttributesScratch.Add(FKeyAttributes());
+				}
 			}
-			else
+		}
+		else
+		{
+			for (int32 i = 0; i < KeyData.Handles.Num(); ++i)
 			{
-				for (int32 i = 0; i < KeyData.Handles.Num(); ++i)
+				const FKeyAttributes Attributes = KeyData.Attributes[i];
+				if (Attributes.HasLeaveTangent())
 				{
-					const FKeyAttributes Attributes = KeyData.Attributes[i];
-					if (Attributes.HasLeaveTangent())
+					const float LeaveTangent = Attributes.GetLeaveTangent();
+
+					if (Attributes.HasTangentWeightMode() && Attributes.HasLeaveTangentWeight() &&
+						(Attributes.GetTangentWeightMode() == RCTWM_WeightedBoth || Attributes.GetTangentWeightMode() == RCTWM_WeightedLeave))
 					{
-						const float LeaveTangent = Attributes.GetLeaveTangent();
-						const FVector2D  KeyScreenPosition = FVector2D(ScreenSpace.SecondsToScreen(KeyPositions[i].InputValue), ScreenSpace.ValueToScreen(KeyPositions[i].OutputValue));
-						if (Attributes.HasTangentWeightMode() && Attributes.HasLeaveTangentWeight() &&
-							(Attributes.GetTangentWeightMode() == RCTWM_WeightedBoth || Attributes.GetTangentWeightMode() == RCTWM_WeightedLeave))
-						{
-							float Tangent, Weight;
-							FVector2D TangentScreenPosition = CurveEditor->GetTangentPositionInScreenSpace(KeyScreenPosition, LeaveTangent, Attributes.GetLeaveTangentWeight());
-							TangentScreenPosition.X += PixelDelta.X;
-							TangentScreenPosition.Y += PixelDelta.Y;
-							CurveEditor->GetTangentAndWeightFromScreenPosition(KeyScreenPosition, TangentScreenPosition, Tangent, Weight);
-								
-							FKeyAttributes NewAttributes;
-							NewAttributes.SetLeaveTangent(Tangent);
-							NewAttributes.SetLeaveTangentWeight(Weight);
-							NewKeyAttributesScratch.Add(NewAttributes);
-						}
-						else
+						FVector2D TangentOffset = CurveEditor::ComputeScreenSpaceTangentOffset(CurveSpace, LeaveTangent, Attributes.GetArriveTangentWeight());
+						TangentOffset += PixelDelta;
+
+						if (MouseEvent.IsShiftDown())
 						{
-							const float PixelLength = 60.0f;
-							FVector2D TangentScreenPosition = FCurveEditor::GetVectorFromSlopeAndLength(LeaveTangent * -DisplayRatio, PixelLength) + KeyScreenPosition;
-							TangentScreenPosition += PixelDelta;
-							TangentScreenPosition -= KeyScreenPosition;
-							const float Tangent = (-TangentScreenPosition.Y / TangentScreenPosition.X) / DisplayRatio;
-							FKeyAttributes NewAttributes;
-							NewAttributes.SetLeaveTangent(Tangent);
-							NewKeyAttributesScratch.Add(NewAttributes);
+							TangentOffset = RoundTrajectory(TangentOffset);
 						}
+
+						// We prevent the handle from crossing over the 0 point... the code handles it but it creates an ugly pop in your curve and it lets your Arrive tangents become Leave tangents which defeats the point.
+						TangentOffset.X = FMath::Max(TangentOffset.X, CurveEditorDragOperation::TangentCrossoverThresholdPx);
+
+						float Tangent, Weight;
+						CurveEditor::TangentAndWeightFromOffset(CurveSpace, TangentOffset, Tangent, Weight);
+
+						FKeyAttributes NewAttributes;
+						NewAttributes.SetLeaveTangent(Tangent);
+						NewAttributes.SetLeaveTangentWeight(Weight);
+						NewKeyAttributesScratch.Add(NewAttributes);
 					}
-					else //still need to add since expect attributes to equal num of selected
+					else
 					{
-						NewKeyAttributesScratch.Add(FKeyAttributes());
+						const float PixelLength = 60.0f;
+						FVector2D TangentOffset = CurveEditor::GetVectorFromSlopeAndLength(LeaveTangent * -DisplayRatio, PixelLength);
+						TangentOffset += PixelDelta;
+
+						if (MouseEvent.IsShiftDown())
+						{
+							TangentOffset = RoundTrajectory(TangentOffset);
+						}
+
+						// We prevent the handle from crossing over the 0 point... the code handles it but it creates an ugly pop in your curve and it lets your Arrive tangents become Leave tangents which defeats the point.
+						TangentOffset.X = FMath::Max(TangentOffset.X, CurveEditorDragOperation::TangentCrossoverThresholdPx);
+
+						const float Tangent = (-TangentOffset.Y / TangentOffset.X) / DisplayRatio;
+
+						FKeyAttributes NewAttributes;
+						NewAttributes.SetLeaveTangent(Tangent);
+						NewKeyAttributesScratch.Add(NewAttributes);
 					}
 				}
+				else //still need to add since expect attributes to equal num of selected
+				{
+					NewKeyAttributesScratch.Add(FKeyAttributes());
+				}
 			}
-			Curve->SetKeyAttributes(KeyData.Handles, NewKeyAttributesScratch);
 		}
+		Curve->SetKeyAttributes(KeyData.Handles, NewKeyAttributesScratch);
 	}
 }
 
@@ -168,6 +212,14 @@ void FCurveEditorDragOperation_Tangent::OnCancelDrag()
 			Curve->SetKeyAttributes(KeyData.Handles, KeyData.Attributes);
 		}
 	}
+
+	CurveEditor->SuppressBoundTransformUpdates(false);
+}
+
+void FCurveEditorDragOperation_Tangent::OnEndDrag(FVector2D InitialPosition, FVector2D CurrentPosition, const FPointerEvent& MouseEvent)
+{
+	ICurveEditorKeyDragOperation::OnEndDrag(InitialPosition, CurrentPosition, MouseEvent);
+	CurveEditor->SuppressBoundTransformUpdates(false);
 }
 
 FVector2D FCurveEditorDragOperation_Tangent::RoundTrajectory(FVector2D Delta)
diff --git a/Engine/Source/Editor/CurveEditor/Private/DragOperations/CurveEditorDragOperation_Tangent.h b/Engine/Source/Editor/CurveEditor/Private/DragOperations/CurveEditorDragOperation_Tangent.h
--- a/Engine/Source/Editor/CurveEditor/Private/DragOperations/CurveEditorDragOperation_Tangent.h
+++ b/Engine/Source/Editor/CurveEditor/Private/DragOperations/CurveEditorDragOperation_Tangent.h
@@ -15,6 +15,7 @@ public:
 	virtual void OnInitialize(FCurveEditor* InCurveEditor, const TOptional<FCurvePointHandle>& InCardinalPoint) override;
 	virtual void OnBeginDrag(FVector2D InitialPosition, FVector2D CurrentPosition, const FPointerEvent& MouseEvent) override;
 	virtual void OnDrag(FVector2D InitialPosition, FVector2D CurrentPosition, const FPointerEvent& MouseEvent) override;
+	virtual void OnEndDrag(FVector2D InitialPosition, FVector2D CurrentPosition, const FPointerEvent& MouseEvent) override;
 	virtual void OnCancelDrag() override;
 
 private:
diff --git a/Engine/Source/Editor/CurveEditor/Private/DragOperations/CurveEditorDragOperation_Zoom.cpp b/Engine/Source/Editor/CurveEditor/Private/DragOperations/CurveEditorDragOperation_Zoom.cpp
--- a/Engine/Source/Editor/CurveEditor/Private/DragOperations/CurveEditorDragOperation_Zoom.cpp
+++ b/Engine/Source/Editor/CurveEditor/Private/DragOperations/CurveEditorDragOperation_Zoom.cpp
@@ -3,45 +3,68 @@
 #include "CurveEditorDragOperation_Zoom.h"
 #include "CurveEditorScreenSpace.h"
 #include "CurveEditor.h"
+#include "SCurveEditorView.h"
+#include "SCurveEditorPanel.h"
 
-FCurveEditorDragOperation_Zoom::FCurveEditorDragOperation_Zoom(FCurveEditor* InCurveEditor)
+FCurveEditorDragOperation_Zoom::FCurveEditorDragOperation_Zoom(FCurveEditor* InCurveEditor, TSharedPtr<SCurveEditorView> InOptionalView)
 	: CurveEditor(InCurveEditor)
-{
-}
+	, OptionalView(InOptionalView)
+{}
 
 void FCurveEditorDragOperation_Zoom::OnBeginDrag(FVector2D InitialPosition, FVector2D CurrentPosition, const FPointerEvent& MouseEvent)
 {
-	FCurveEditorScreenSpace ScreenSpace = CurveEditor->GetScreenSpace();
+	FCurveEditorScreenSpaceH InputSpace  = OptionalView ? OptionalView->GetViewSpace() : CurveEditor->GetPanelInputSpace();
 
-	ZoomFactor = InitialPosition / ScreenSpace.GetPhysicalSize();
+	ZoomFactor.X = InitialPosition.X / InputSpace.GetPhysicalWidth();
+	OriginalInputRange = (InputSpace.GetInputMax() - InputSpace.GetInputMin());
+	ZoomOriginX = InputSpace.GetInputMin() + OriginalInputRange * ZoomFactor.X;
 
-	OriginalInputRange = ( ScreenSpace.GetInputMax() - ScreenSpace.GetInputMin() );
-	OriginalOutputRange = ( ScreenSpace.GetOutputMax() - ScreenSpace.GetOutputMin() );
+	if (OptionalView)
+	{
+		FCurveEditorScreenSpaceV OutputSpace = OptionalView->GetViewSpace();
 
-	ZoomOriginX = ScreenSpace.GetInputMin() + OriginalInputRange * ZoomFactor.X;
-	ZoomOriginY = ScreenSpace.GetOutputMin() + OriginalOutputRange * ZoomFactor.Y;
+		ZoomFactor.Y = InitialPosition.Y / OutputSpace.GetPhysicalHeight();
+		OriginalOutputRange = (OutputSpace.GetOutputMax() - OutputSpace.GetOutputMin());
+		ZoomOriginY = OutputSpace.GetOutputMin() + OriginalOutputRange * (1.f - ZoomFactor.Y);
+	}
 }
 
 void FCurveEditorDragOperation_Zoom::OnDrag(FVector2D InitialPosition, FVector2D CurrentPosition, const FPointerEvent& MouseEvent)
 {
-	FVector2D PixelDelta = GetLockedMousePosition(InitialPosition, CurrentPosition, MouseEvent) - InitialPosition;
-
-	FCurveEditorScreenSpace ScreenSpace = CurveEditor->GetScreenSpace();
+	FVector2D PixelDelta = CurrentPosition - InitialPosition;
 
 	double ClampRange = 1e9f;
 
-	double DiffX = double(PixelDelta.X) / (ScreenSpace.GetPhysicalWidth() / OriginalInputRange);
-	double DiffY = double(PixelDelta.Y) / (ScreenSpace.GetPhysicalHeight() / OriginalOutputRange);
+	// Zoom input range
+	FCurveEditorScreenSpaceH InputSpace  = OptionalView ? OptionalView->GetViewSpace() : CurveEditor->GetPanelInputSpace();
 
-	double NewInputRange = OriginalInputRange + DiffX;
-	double InputMin  = FMath::Clamp<double>(ZoomOriginX - NewInputRange * ZoomFactor.X, -ClampRange, ClampRange);
-	double InputMax  = FMath::Clamp<double>(ZoomOriginX + NewInputRange * (1.f - ZoomFactor.X), InputMin, ClampRange);
+	double DiffX = double(PixelDelta.X) / (InputSpace.GetPhysicalWidth() / OriginalInputRange);
 
-
-	double NewOutputRange = OriginalOutputRange + DiffY;
-	double OutputMin  = FMath::Clamp<double>(ZoomOriginY - NewOutputRange * ZoomFactor.Y, -ClampRange, ClampRange);
-	double OutputMax  = FMath::Clamp<double>(ZoomOriginY + NewOutputRange * (1.f - ZoomFactor.Y), OutputMin, ClampRange);
+	// This flips the horizontal zoom to match existing DCC tools
+	double NewInputRange = OriginalInputRange - DiffX;
+	double InputMin = FMath::Clamp<double>(ZoomOriginX - NewInputRange * ZoomFactor.X, -ClampRange, ClampRange);
+	double InputMax = FMath::Clamp<double>(ZoomOriginX + NewInputRange * (1.f - ZoomFactor.X), InputMin, ClampRange);
 
 	CurveEditor->GetBounds().SetInputBounds(InputMin, InputMax);
-	CurveEditor->GetBounds().SetOutputBounds(OutputMin, OutputMax);
-}
\ No newline at end of file
+
+	// Zoom output range
+	if (OptionalView.IsValid())
+	{
+		FCurveEditorScreenSpaceV ViewSpace = OptionalView->GetViewSpace();
+
+		double DiffY = double(PixelDelta.Y) / (ViewSpace.GetPhysicalHeight() / OriginalOutputRange);
+		double NewOutputRange = OriginalOutputRange + DiffY;
+
+		// If they're holding a shift they can scale on both axis at once, non-proportionally
+		if (!MouseEvent.IsShiftDown())
+		{
+			// By default, do proportional zoom
+			NewOutputRange = (NewInputRange / OriginalInputRange) * OriginalOutputRange;
+		}
+
+		double OutputMin = FMath::Clamp<double>(ZoomOriginY - NewOutputRange * (1.f - ZoomFactor.Y), -ClampRange, ClampRange);
+		double OutputMax = FMath::Clamp<double>(ZoomOriginY + NewOutputRange * (ZoomFactor.Y), OutputMin, ClampRange);
+
+		OptionalView->SetOutputBounds(OutputMin, OutputMax);
+	}
+}
diff --git a/Engine/Source/Editor/CurveEditor/Private/DragOperations/CurveEditorDragOperation_Zoom.h b/Engine/Source/Editor/CurveEditor/Private/DragOperations/CurveEditorDragOperation_Zoom.h
--- a/Engine/Source/Editor/CurveEditor/Private/DragOperations/CurveEditorDragOperation_Zoom.h
+++ b/Engine/Source/Editor/CurveEditor/Private/DragOperations/CurveEditorDragOperation_Zoom.h
@@ -5,11 +5,12 @@
 #include "ICurveEditorDragOperation.h"
 
 class FCurveEditor;
+class SCurveEditorView;
 
 class FCurveEditorDragOperation_Zoom : public ICurveEditorDragOperation
 {
 public:
-	FCurveEditorDragOperation_Zoom(FCurveEditor* CurveEditor);
+	FCurveEditorDragOperation_Zoom(FCurveEditor* CurveEditor, TSharedPtr<SCurveEditorView> InOptionalView);
 
 	virtual void OnBeginDrag(FVector2D InitialPosition, FVector2D CurrentPosition, const FPointerEvent& MouseEvent) override;
 	virtual void OnDrag(FVector2D InitialPosition, FVector2D CurrentPosition, const FPointerEvent& MouseEvent) override;
@@ -17,6 +18,7 @@ public:
 private:
 
 	FCurveEditor* CurveEditor;
+	TSharedPtr<SCurveEditorView> OptionalView;
 
 	FVector2D ZoomFactor;
 	double ZoomOriginX, ZoomOriginY;
diff --git a/Engine/Source/Editor/CurveEditor/Private/DragOperations/ICurveEditorDragOperation.cpp b/Engine/Source/Editor/CurveEditor/Private/DragOperations/ICurveEditorDragOperation.cpp
--- a/Engine/Source/Editor/CurveEditor/Private/DragOperations/ICurveEditorDragOperation.cpp
+++ b/Engine/Source/Editor/CurveEditor/Private/DragOperations/ICurveEditorDragOperation.cpp
@@ -6,7 +6,6 @@
 
 void ICurveEditorDragOperation::BeginDrag(FVector2D InitialPosition, FVector2D CurrentPosition, const FPointerEvent& MouseEvent)
 {
-	GetLockedMousePosition(InitialPosition, CurrentPosition, MouseEvent);
 	OnBeginDrag(InitialPosition, CurrentPosition, MouseEvent);
 }
 
@@ -20,9 +19,9 @@ void ICurveEditorDragOperation::EndDrag(FVector2D InitialPosition, FVector2D Cur
 	OnEndDrag(InitialPosition, CurrentPosition, MouseEvent);
 }
 
-int32 ICurveEditorDragOperation::Paint(const FGeometry& AllottedGeometry, FSlateWindowElementList& OutDrawElements, int32 LayerId)
+void ICurveEditorDragOperation::Paint(const FGeometry& AllottedGeometry, FSlateWindowElementList& OutDrawElements, int32 PaintOnLayerId)
 {
-	return OnPaint(AllottedGeometry, OutDrawElements, LayerId);
+	OnPaint(AllottedGeometry, OutDrawElements, PaintOnLayerId);
 }
 
 void ICurveEditorDragOperation::CancelDrag()
@@ -30,29 +29,6 @@ void ICurveEditorDragOperation::CancelDrag()
 	OnCancelDrag();
 }
 
-FVector2D ICurveEditorDragOperation::GetLockedMousePosition(FVector2D InitialPosition, FVector2D CurrentPosition, const FPointerEvent& MouseEvent)
-{
-	if (MouseEvent.IsShiftDown())
-	{
-		if (MouseLockVector == FVector2D::UnitVector)
-		{
-			if (FMath::Abs(CurrentPosition.Y - InitialPosition.Y) <= FMath::Abs(CurrentPosition.X - InitialPosition.X))
-			{
-				MouseLockVector.Y = 0.f;
-			}
-			else
-			{
-				MouseLockVector.X = 0.f;
-			}
-		}
-	}
-	else
-	{
-		MouseLockVector = FVector2D::UnitVector;
-	}
-	return InitialPosition + (CurrentPosition - InitialPosition) * MouseLockVector;
-}
-
 void ICurveEditorKeyDragOperation::Initialize(FCurveEditor* InCurveEditor, const TOptional<FCurvePointHandle>& CardinalPoint)
 {
 	SnapMetrics = InCurveEditor->GetSnapMetrics();
diff --git a/Engine/Source/Editor/CurveEditor/Private/Filters/CurveEditorBakeFilter.cpp b/Engine/Source/Editor/CurveEditor/Private/Filters/CurveEditorBakeFilter.cpp
new file mode 100644
--- /dev/null
+++ b/Engine/Source/Editor/CurveEditor/Private/Filters/CurveEditorBakeFilter.cpp
@@ -0,0 +1,120 @@
+// Copyright 1998-2018 Epic Games, Inc. All Rights Reserved.
+
+#include "Filters/CurveEditorBakeFilter.h"
+#include "CoreMinimal.h"
+#include "ScopedTransaction.h"
+#include "Curves/KeyHandle.h"
+#include "CurveDataAbstraction.h"
+#include "CurveEditor.h"
+#include "CurveEditorSelection.h"
+#include "CurveModel.h"
+#include "Misc/FrameRate.h"
+#include "CurveEditorSnapMetrics.h"
+#include "Containers/ArrayView.h"
+
+void UCurveEditorBakeFilter::ApplyFilter_Impl(TSharedRef<FCurveEditor> InCurveEditor, const TMap<FCurveModelID, FKeyHandleSet>& InKeysToOperateOn, TMap<FCurveModelID, FKeyHandleSet>& OutKeysToSelect)
+{
+	TArray<FKeyHandle> KeyHandles;
+	TArray<FKeyPosition> SelectedKeyPositions;
+
+	TArray<FKeyPosition> NewKeyPositions;
+	TArray<FKeyAttributes> NewKeyAttributes;
+
+	// Look at our Input Snap rate to determine how far apart the baked keys are.
+	FFrameRate BakeRate = InCurveEditor->GetSnapMetrics().InputSnapRate;
+
+	for (const TTuple<FCurveModelID, FKeyHandleSet>& Pair : InKeysToOperateOn)
+	{
+		FCurveModel* Curve = InCurveEditor->FindCurve(Pair.Key);
+
+		if (!Curve)
+		{
+			continue;
+		}
+
+		KeyHandles.Reset(Pair.Value.Num());
+		KeyHandles.Append(Pair.Value.AsArray().GetData(), Pair.Value.Num());
+
+		// Get all the selected keys
+		SelectedKeyPositions.SetNum(KeyHandles.Num());
+		Curve->GetKeyPositions(KeyHandles, SelectedKeyPositions);
+
+		// Find the hull of the range of the selected keys
+		double MinKey = TNumericLimits<double>::Max(), MaxKey = TNumericLimits<double>::Lowest();
+		for (FKeyPosition Key : SelectedKeyPositions)
+		{
+			MinKey = FMath::Min(Key.InputValue, MinKey);
+			MaxKey = FMath::Max(Key.InputValue, MaxKey);
+		}
+
+		// Get all keys that exist between the time range
+		KeyHandles.Reset();
+		Curve->GetKeys(*InCurveEditor, MinKey, MaxKey, TNumericLimits<double>::Lowest(), TNumericLimits<double>::Max(), KeyHandles);
+
+		if (KeyHandles.Num() > 1)
+		{
+			// Determine new times for new keys
+			double Interval = bUseSnapRateForInterval ? BakeRate.AsInterval() : BakeInterval;
+			int32 NumKeysToAdd = FMath::FloorToInt((MaxKey - MinKey) / Interval) + 1;
+
+			NewKeyPositions.Reset(NumKeysToAdd);
+			NewKeyAttributes.Reset(NumKeysToAdd);
+
+			for (int32 KeyIndex = 1; KeyIndex < NumKeysToAdd; ++KeyIndex)
+			{
+				double NewTime = MinKey + KeyIndex * Interval;
+
+				// Don't create the last key if it has the same value as the last key. This avoids duplicates of the last
+				// key, but improves behavior when the interval doesn't align.
+				if (KeyIndex == NumKeysToAdd - 1)
+				{
+					if (FMath::IsNearlyEqual(NewTime, SelectedKeyPositions.Last().InputValue))
+					{
+						continue;
+					}
+				}
+				FKeyPosition CurrentKey(NewTime, 0.0);
+				if (Curve->Evaluate(CurrentKey.InputValue, CurrentKey.OutputValue))
+				{
+					NewKeyPositions.Add(CurrentKey);
+					NewKeyAttributes.Add(FKeyAttributes().SetInterpMode(RCIM_Linear));
+				}
+			}
+
+			Curve->Modify();
+
+			// We want to store the newly added ones so we can add them to the users selection to mimic
+			// their selection pre-baking.
+			TArray<TOptional<FKeyHandle>> NewKeyHandles;
+			NewKeyHandles.SetNumUninitialized(NewKeyPositions.Num());
+
+			TArrayView<TOptional<FKeyHandle>> NewKeyHandlesView = MakeArrayView(NewKeyHandles);
+			FKeyHandleSet& OutHandleSet = OutKeysToSelect.Add(Pair.Key);
+
+			// Manually add the first and last keys to the out set before we remove them below.
+			OutHandleSet.Add(KeyHandles[0]);
+			OutHandleSet.Add(KeyHandles.Last(0));
+
+			// We need to leave the first and the last key of the selection alone for two reasons;
+			// 1. Undo/Redo works better as selections aren't transacted so they don't handle keys being removed/re-added.
+			// by leaving the keys alone, they survive undo/redo operations.
+			// 2. With high interval settings, the shape of the curve can change drastically because the next interval would fall right
+			// outside the last key. If we remove the keys, then the shape changes in these cases. We avoid putting a new key on the last
+			// key (if it exists), to avoid any duplication. This preserves the shape of the curves better in most test cases.
+			KeyHandles.RemoveAt(0);
+			KeyHandles.RemoveAt(KeyHandles.Num() - 1);
+
+			// Remove all the old in-between keys and add the new ones
+			Curve->RemoveKeys(KeyHandles);
+			Curve->AddKeys(NewKeyPositions, NewKeyAttributes, &NewKeyHandlesView);
+
+			for (const TOptional<FKeyHandle> Handle : NewKeyHandlesView)
+			{
+				if (Handle.IsSet())
+				{
+					OutHandleSet.Add(Handle.GetValue());
+				}
+			}
+		}
+	}
+}
diff --git a/Engine/Source/Editor/CurveEditor/Private/RichCurveEditorModel.cpp b/Engine/Source/Editor/CurveEditor/Private/RichCurveEditorModel.cpp
new file mode 100644
--- /dev/null
+++ b/Engine/Source/Editor/CurveEditor/Private/RichCurveEditorModel.cpp
@@ -0,0 +1,524 @@
+// Copyright 1998-2018 Epic Games, Inc. All Rights Reserved.
+
+#include "RichCurveEditorModel.h"
+#include "RichCurveKeyProxy.h"
+#include "Math/Vector2D.h"
+#include "HAL/PlatformMath.h"
+#include "Curves/RichCurve.h"
+#include "CurveDrawInfo.h"
+#include "CurveDataAbstraction.h"
+#include "CurveEditor.h"
+#include "CurveEditorScreenSpace.h"
+#include "CurveEditorSnapMetrics.h"
+#include "EditorStyleSet.h"
+#include "UObject/Package.h"
+
+FRichCurveEditorModel::FRichCurveEditorModel(FRichCurve* InRichCurve, UObject* InOwner)
+	: RichCurve(InRichCurve), WeakOwner(InOwner)
+{
+	checkf(RichCurve, TEXT("If is not valid to provide a null rich curve to this class"));
+}
+
+const void* FRichCurveEditorModel::GetCurve() const
+{
+	return RichCurve;
+}
+
+void FRichCurveEditorModel::Modify()
+{
+	if (UObject* Owner = WeakOwner.Get())
+	{
+		Owner->Modify();
+	}
+}
+
+void FRichCurveEditorModel::AddKeys(TArrayView<const FKeyPosition> InKeyPositions, TArrayView<const FKeyAttributes> InKeyAttributes, TArrayView<TOptional<FKeyHandle>>* OutKeyHandles)
+{
+	check(InKeyPositions.Num() == InKeyAttributes.Num() && (!OutKeyHandles || OutKeyHandles->Num() == InKeyPositions.Num()));
+
+	UObject* Owner = WeakOwner.Get();
+	if (Owner)
+	{
+		Owner->Modify();
+
+		for (int32 Index = 0; Index < InKeyPositions.Num(); ++Index)
+		{
+			FKeyPosition   Position   = InKeyPositions[Index];
+			FKeyAttributes Attributes = InKeyAttributes[Index];
+
+			FKeyHandle     NewHandle = RichCurve->AddKey(Position.InputValue, Position.OutputValue);
+			FRichCurveKey* NewKey    = &RichCurve->GetKey(NewHandle);
+
+			if (Attributes.HasInterpMode())          { NewKey->InterpMode          = Attributes.GetInterpMode();    }
+			if (Attributes.HasTangentMode())         { NewKey->TangentMode         = Attributes.GetTangentMode();   }
+			if (Attributes.HasArriveTangent())       { NewKey->ArriveTangent       = Attributes.GetArriveTangent(); }
+			if (Attributes.HasLeaveTangent())        { NewKey->LeaveTangent        = Attributes.GetLeaveTangent();  }
+			if (Attributes.HasArriveTangentWeight()) { NewKey->ArriveTangentWeight = Attributes.GetArriveTangent(); }
+			if (Attributes.HasLeaveTangentWeight())  { NewKey->LeaveTangentWeight  = Attributes.GetLeaveTangent();  }
+
+			if (OutKeyHandles)
+			{
+				(*OutKeyHandles)[Index] = NewHandle;
+			}
+		}
+
+		RichCurve->AutoSetTangents();
+	}
+}
+
+bool FRichCurveEditorModel::Evaluate(double Time, double& OutValue) const
+{
+	if (UObject* Owner = WeakOwner.Get())
+	{
+		OutValue = RichCurve->Eval(Time);
+		return true;
+	}
+
+	return false;
+}
+
+void FRichCurveEditorModel::RemoveKeys(TArrayView<const FKeyHandle> InKeys)
+{
+	if (UObject* Owner = WeakOwner.Get())
+	{
+		Owner->Modify();
+		for (FKeyHandle Handle : InKeys)
+		{
+			RichCurve->DeleteKey(Handle);
+		}
+	}
+}
+
+void RefineCurvePoints(const FRichCurve* RichCurve, double TimeThreshold, float ValueThreshold, TArray<TTuple<double, double>>& InOutPoints)
+{
+	const float InterpTimes[] = { 0.25f, 0.5f, 0.6f };
+
+	for (int32 Index = 0; Index < InOutPoints.Num() - 1; ++Index)
+	{
+		TTuple<double, double> Lower = InOutPoints[Index];
+		TTuple<double, double> Upper = InOutPoints[Index + 1];
+
+		if ((Upper.Get<0>() - Lower.Get<0>()) >= TimeThreshold)
+		{
+			bool bSegmentIsLinear = true;
+
+			TTuple<double, double> Evaluated[ARRAY_COUNT(InterpTimes)];
+
+			for (int32 InterpIndex = 0; InterpIndex < ARRAY_COUNT(InterpTimes); ++InterpIndex)
+			{
+				double& EvalTime  = Evaluated[InterpIndex].Get<0>();
+
+				EvalTime = FMath::Lerp(Lower.Get<0>(), Upper.Get<0>(), InterpTimes[InterpIndex]);
+
+				float Value = RichCurve->Eval(EvalTime);
+
+				const float LinearValue = FMath::Lerp(Lower.Get<1>(), Upper.Get<1>(), InterpTimes[InterpIndex]);
+				if (bSegmentIsLinear)
+				{
+					bSegmentIsLinear = FMath::IsNearlyEqual(Value, LinearValue, ValueThreshold);
+				}
+
+				Evaluated[InterpIndex].Get<1>() = Value;
+			}
+
+			if (!bSegmentIsLinear)
+			{
+				// Add the point
+				InOutPoints.Insert(Evaluated, ARRAY_COUNT(Evaluated), Index+1);
+				--Index;
+			}
+		}
+	}
+}
+
+void FRichCurveEditorModel::DrawCurve(const FCurveEditor& CurveEditor, const FCurveEditorScreenSpace& ScreenSpace, TArray<TTuple<double, double>>& InOutPoints) const
+{
+	if (UObject* Owner = WeakOwner.Get())
+	{
+		const double StartTimeSeconds = ScreenSpace.GetInputMin();
+		const double EndTimeSeconds   = ScreenSpace.GetInputMax();
+		const double TimeThreshold    = FMath::Max(0.0001, 1.0 / ScreenSpace.PixelsPerInput());
+		const double ValueThreshold   = FMath::Max(0.0001, 1.0 / ScreenSpace.PixelsPerOutput());
+
+		InOutPoints.Add(MakeTuple(StartTimeSeconds, double(RichCurve->Eval(StartTimeSeconds))));
+
+		for (const FRichCurveKey& Key : RichCurve->GetConstRefOfKeys())
+		{
+			if (Key.Time > StartTimeSeconds && Key.Time < EndTimeSeconds)
+			{
+				InOutPoints.Add(MakeTuple(double(Key.Time), double(Key.Value)));
+			}
+		}
+
+		InOutPoints.Add(MakeTuple(EndTimeSeconds, double(RichCurve->Eval(EndTimeSeconds))));
+
+		int32 OldSize = InOutPoints.Num();
+		do
+		{
+			OldSize = InOutPoints.Num();
+			RefineCurvePoints(RichCurve, TimeThreshold, ValueThreshold, InOutPoints);
+		}
+		while(OldSize != InOutPoints.Num());
+	}
+}
+
+void FRichCurveEditorModel::GetKeys(const FCurveEditor& CurveEditor, double MinTime, double MaxTime, double MinValue, double MaxValue, TArray<FKeyHandle>& OutKeyHandles) const
+{
+	if (UObject* Owner = WeakOwner.Get())
+	{
+		for (auto It = RichCurve->GetKeyHandleIterator(); It; ++It)
+		{
+			const FRichCurveKey& Key = RichCurve->GetKey(*It);
+			if (Key.Time >= MinTime && Key.Time <= MaxTime && Key.Value >= MinValue && Key.Value <= MaxValue)
+			{
+				OutKeyHandles.Add(*It);
+			}
+		}
+	}
+}
+
+void FRichCurveEditorModel::GetKeyDrawInfo(ECurvePointType PointType, const FKeyHandle InKeyHandle, FKeyDrawInfo& OutDrawInfo) const
+{
+	if (PointType == ECurvePointType::ArriveTangent || PointType == ECurvePointType::LeaveTangent)
+	{
+		OutDrawInfo.Brush = FEditorStyle::GetBrush("GenericCurveEditor.TangentHandle");
+		OutDrawInfo.ScreenSize = FVector2D(9, 9);
+	}
+	else
+	{
+		// All keys are the same size by default
+		OutDrawInfo.ScreenSize = FVector2D(11, 11);
+
+		ERichCurveInterpMode KeyType = RichCurve->IsKeyHandleValid(InKeyHandle) ? RichCurve->GetKey(InKeyHandle).InterpMode.GetValue() : RCIM_None;
+		switch (KeyType)
+		{
+		case ERichCurveInterpMode::RCIM_Constant:
+			OutDrawInfo.Brush = FEditorStyle::GetBrush("GenericCurveEditor.ConstantKey");
+			OutDrawInfo.Tint = FLinearColor(0, 0.45f, 0.70f);
+			break;
+		case ERichCurveInterpMode::RCIM_Linear:
+			OutDrawInfo.Brush = FEditorStyle::GetBrush("GenericCurveEditor.LinearKey");
+			OutDrawInfo.Tint = FLinearColor(0, 0.62f, 0.46f);
+			break;
+		case ERichCurveInterpMode::RCIM_Cubic:
+			OutDrawInfo.Brush = FEditorStyle::GetBrush("GenericCurveEditor.CubicKey");
+			OutDrawInfo.Tint = FLinearColor::White;
+			break;
+		default:
+			OutDrawInfo.Brush = FEditorStyle::GetBrush("GenericCurveEditor.Key");
+			OutDrawInfo.Tint = FLinearColor::White;
+			break;
+		}
+	}
+}
+
+void FRichCurveEditorModel::GetKeyPositions(TArrayView<const FKeyHandle> InKeys, TArrayView<FKeyPosition> OutKeyPositions) const
+{
+	if (UObject* Owner = WeakOwner.Get())
+	{
+		for (int32 Index = 0; Index < InKeys.Num(); ++Index)
+		{
+			if (RichCurve->IsKeyHandleValid(InKeys[Index]))
+			{
+				const FRichCurveKey& Key = RichCurve->GetKey(InKeys[Index]);
+
+				OutKeyPositions[Index].InputValue  = Key.Time;
+				OutKeyPositions[Index].OutputValue = Key.Value;
+			}
+		}
+	}
+}
+
+void FRichCurveEditorModel::SetKeyPositions(TArrayView<const FKeyHandle> InKeys, TArrayView<const FKeyPosition> InKeyPositions)
+{
+	if (UObject* Owner = WeakOwner.Get())
+	{
+		Owner->Modify();
+
+		for (int32 Index = 0; Index < InKeys.Num(); ++Index)
+		{
+			FKeyHandle Handle = InKeys[Index];
+			if (RichCurve->IsKeyHandleValid(Handle))
+			{
+				// Set key time last so we don't have to worry about the key handle changing
+				RichCurve->GetKey(Handle).Value = InKeyPositions[Index].OutputValue;
+				RichCurve->SetKeyTime(Handle, InKeyPositions[Index].InputValue);
+			}
+		}
+		RichCurve->AutoSetTangents();
+	}
+}
+
+void FRichCurveEditorModel::GetKeyAttributes(TArrayView<const FKeyHandle> InKeys, TArrayView<FKeyAttributes> OutAttributes) const
+{
+	if (UObject* Owner = WeakOwner.Get())
+	{
+		const TArray<FRichCurveKey>& AllKeys = RichCurve->GetConstRefOfKeys();
+		if (AllKeys.Num() == 0)
+		{
+			return;
+		}
+
+		const FRichCurveKey* FirstKey = &AllKeys[0];
+		const FRichCurveKey* LastKey  = &AllKeys.Last();
+
+		for (int32 Index = 0; Index < InKeys.Num(); ++Index)
+		{
+			if (RichCurve->IsKeyHandleValid(InKeys[Index]))
+			{
+				const FRichCurveKey* ThisKey    = &RichCurve->GetKey(InKeys[Index]);
+				FKeyAttributes&      Attributes = OutAttributes[Index];
+
+				Attributes.SetInterpMode(ThisKey->InterpMode);
+
+				if (ThisKey->InterpMode != RCIM_Constant && ThisKey->InterpMode != RCIM_Linear)
+				{
+					Attributes.SetTangentMode(ThisKey->TangentMode);
+					if (ThisKey != FirstKey)
+					{
+						Attributes.SetArriveTangent(ThisKey->ArriveTangent);
+					}
+
+					if (ThisKey != LastKey)
+					{
+						Attributes.SetLeaveTangent(ThisKey->LeaveTangent);
+					}
+					if (ThisKey->InterpMode == RCIM_Cubic)
+					{
+						Attributes.SetTangentWeightMode(ThisKey->TangentWeightMode);
+						if (ThisKey->TangentWeightMode != RCTWM_WeightedNone)
+						{
+							Attributes.SetArriveTangentWeight(ThisKey->ArriveTangentWeight);
+							Attributes.SetLeaveTangentWeight(ThisKey->LeaveTangentWeight);
+						}
+					}
+				}
+			}
+		}
+	}
+}
+
+void FRichCurveEditorModel::SetKeyAttributes(TArrayView<const FKeyHandle> InKeys, TArrayView<const FKeyAttributes> InAttributes)
+{
+	if (UObject* Owner = WeakOwner.Get())
+	{
+		const TArray<FRichCurveKey>& AllKeys = RichCurve->GetConstRefOfKeys();
+		if (AllKeys.Num() == 0)
+		{
+			return;
+		}
+
+		Owner->Modify();
+
+		const FRichCurveKey* FirstKey = &AllKeys[0];
+		const FRichCurveKey* LastKey  = &AllKeys.Last();
+
+		bool bAutoSetTangents = false;
+
+		for (int32 Index = 0; Index < InKeys.Num(); ++Index)
+		{
+			if (RichCurve->IsKeyHandleValid(InKeys[Index]))
+			{
+				FRichCurveKey*        ThisKey    = &RichCurve->GetKey(InKeys[Index]);
+				const FKeyAttributes& Attributes = InAttributes[Index];
+
+				if (Attributes.HasInterpMode())    { ThisKey->InterpMode  = Attributes.GetInterpMode();  bAutoSetTangents = true; }
+				if (Attributes.HasTangentMode())
+				{
+					ThisKey->TangentMode = Attributes.GetTangentMode();
+					if (ThisKey->TangentMode == RCTM_Auto)
+					{
+						ThisKey->TangentWeightMode = RCTWM_WeightedNone;
+					}
+					bAutoSetTangents = true;
+				}
+				if (Attributes.HasTangentWeightMode()) 
+				{ 
+					if (ThisKey->TangentWeightMode == RCTWM_WeightedNone) //set tangent weights to default use
+					{
+						const float OneThird = 1.0f / 3.0f;
+
+						//calculate a tangent weight based upon tangent and time difference
+						//calculate arrive tangent weight
+						if (ThisKey != FirstKey)
+						{
+							const float Y = ThisKey->ArriveTangent;
+							ThisKey->ArriveTangentWeight = FMath::Sqrt(1.f + Y*Y) * OneThird;
+						}
+						//calculate leave weight
+						if(ThisKey != LastKey)
+						{
+							const float Y = ThisKey->LeaveTangent;
+							ThisKey->LeaveTangentWeight = FMath::Sqrt(1.f + Y*Y) * OneThird;
+						}
+					}
+					ThisKey->TangentWeightMode = Attributes.GetTangentWeightMode();
+
+					if( ThisKey->TangentWeightMode != RCTWM_WeightedNone )
+					{
+						if (ThisKey->TangentMode != RCTM_User && ThisKey->TangentMode != RCTM_Break)
+						{
+							ThisKey->TangentMode = RCTM_User;
+						}
+					}
+				}
+
+				if (Attributes.HasArriveTangent())
+				{
+					if (ThisKey->TangentMode == RCTM_Auto)
+					{
+						ThisKey->TangentMode = RCTM_User;
+						ThisKey->TangentWeightMode = RCTWM_WeightedNone;
+					}
+
+					ThisKey->ArriveTangent = Attributes.GetArriveTangent();
+					if (ThisKey->InterpMode == RCIM_Cubic && ThisKey->TangentMode != RCTM_Break)
+					{
+						ThisKey->LeaveTangent = ThisKey->ArriveTangent;
+					}
+				}
+
+				if (Attributes.HasLeaveTangent())
+				{
+					if (ThisKey->TangentMode == RCTM_Auto)
+					{
+						ThisKey->TangentMode = RCTM_User;
+						ThisKey->TangentWeightMode = RCTWM_WeightedNone;
+					}
+
+					ThisKey->LeaveTangent = Attributes.GetLeaveTangent();
+					if (ThisKey->InterpMode == RCIM_Cubic && ThisKey->TangentMode != RCTM_Break)
+					{
+						ThisKey->ArriveTangent = ThisKey->LeaveTangent;
+					}
+				}
+
+				if (Attributes.HasArriveTangentWeight())
+				{
+					if (ThisKey->TangentMode == RCTM_Auto)
+					{
+						ThisKey->TangentMode = RCTM_User;
+						ThisKey->TangentWeightMode = RCTWM_WeightedNone;
+					}
+
+					ThisKey->ArriveTangentWeight = Attributes.GetArriveTangentWeight();
+					if (ThisKey->InterpMode == RCIM_Cubic && ThisKey->TangentMode != RCTM_Break)
+					{
+						ThisKey->LeaveTangentWeight = ThisKey->ArriveTangentWeight;
+					}
+				}
+
+				if (Attributes.HasLeaveTangentWeight())
+				{
+				
+					if (ThisKey->TangentMode == RCTM_Auto)
+					{
+						ThisKey->TangentMode = RCTM_User;
+						ThisKey->TangentWeightMode = RCTWM_WeightedNone;
+					}
+
+					ThisKey->LeaveTangentWeight = Attributes.GetLeaveTangentWeight();
+					if (ThisKey->InterpMode == RCIM_Cubic && ThisKey->TangentMode != RCTM_Break)
+					{
+						ThisKey->ArriveTangentWeight = ThisKey->LeaveTangentWeight;
+					}
+				}
+			}
+		}
+
+		if (bAutoSetTangents)
+		{
+			RichCurve->AutoSetTangents();
+		}
+	}
+}
+
+void FRichCurveEditorModel::GetCurveAttributes(FCurveAttributes& OutCurveAttributes) const
+{
+	if (UObject* Owner = WeakOwner.Get())
+	{
+		OutCurveAttributes.SetPreExtrapolation(RichCurve->PreInfinityExtrap);
+		OutCurveAttributes.SetPostExtrapolation(RichCurve->PostInfinityExtrap);
+	}
+}
+
+void FRichCurveEditorModel::SetCurveAttributes(const FCurveAttributes& InCurveAttributes)
+{
+	if (UObject* Owner = WeakOwner.Get())
+	{
+		Owner->Modify();
+
+		if (InCurveAttributes.HasPreExtrapolation())
+		{
+			RichCurve->PreInfinityExtrap = InCurveAttributes.GetPreExtrapolation();
+		}
+
+		if (InCurveAttributes.HasPostExtrapolation())
+		{
+			RichCurve->PostInfinityExtrap = InCurveAttributes.GetPostExtrapolation();
+		}
+	}
+}
+
+void FRichCurveEditorModel::CreateKeyProxies(TArrayView<const FKeyHandle> InKeyHandles, TArrayView<UObject*> OutObjects)
+{
+	for (int32 Index = 0; Index < InKeyHandles.Num(); ++Index)
+	{
+		URichCurveKeyProxy* NewProxy = NewObject<URichCurveKeyProxy>(GetTransientPackage(), NAME_None);
+
+		NewProxy->Initialize(InKeyHandles[Index], RichCurve, WeakOwner);
+		OutObjects[Index] = NewProxy;
+	}
+}
+
+void FRichCurveEditorModel::GetTimeRange(double& MinTime, double& MaxTime) const
+{
+	if (UObject* Owner = WeakOwner.Get())
+	{
+		float MinTimeFloat = 0.f, MaxTimeFloat = 0.f;
+		RichCurve->GetTimeRange(MinTimeFloat, MaxTimeFloat);
+
+		MinTime = MinTimeFloat;
+		MaxTime = MaxTimeFloat;
+	}
+}
+
+void FRichCurveEditorModel::GetValueRange(double& MinValue, double& MaxValue) const
+{
+	if (UObject* Owner = WeakOwner.Get())
+	{
+		float MinValueFloat = 0.f, MaxValueFloat = 0.f;
+		RichCurve->GetValueRange(MinValueFloat, MaxValueFloat);
+
+		MinValue = MinValueFloat;
+		MaxValue = MaxValueFloat;
+	}
+}
+
+int32 FRichCurveEditorModel::GetNumKeys() const
+{
+	return RichCurve->GetNumKeys();
+}
+
+void FRichCurveEditorModel::GetNeighboringKeys(const FKeyHandle InKeyHandle, TOptional<FKeyHandle>& OutPreviousKeyHandle, TOptional<FKeyHandle>& OutNextKeyHandle) const
+{
+	if (UObject* Owner = WeakOwner.Get())
+	{
+		if (RichCurve->IsKeyHandleValid(InKeyHandle))
+		{
+			FKeyHandle NextKeyHandle = RichCurve->GetNextKey(InKeyHandle);
+
+			if (RichCurve->IsKeyHandleValid(NextKeyHandle))
+			{
+				OutNextKeyHandle = NextKeyHandle;
+			}
+
+			FKeyHandle PreviousKeyHandle = RichCurve->GetPreviousKey(InKeyHandle);
+
+			if (RichCurve->IsKeyHandleValid(PreviousKeyHandle))
+			{
+				OutPreviousKeyHandle = PreviousKeyHandle;
+			}
+		}
+	}
+}
\ No newline at end of file
diff --git a/Engine/Source/Editor/CurveEditor/Private/RichCurveKeyProxy.h b/Engine/Source/Editor/CurveEditor/Private/RichCurveKeyProxy.h
new file mode 100644
--- /dev/null
+++ b/Engine/Source/Editor/CurveEditor/Private/RichCurveKeyProxy.h
@@ -0,0 +1,84 @@
+// Copyright 1998-2018 Epic Games, Inc. All Rights Reserved.
+
+#pragma once
+
+#include "CoreTypes.h"
+#include "UObject/WeakObjectPtr.h"
+#include "UObject/WeakObjectPtrTemplates.h"
+#include "CurveEditorKeyProxy.h"
+#include "Curves/RichCurve.h"
+
+#include "RichCurveKeyProxy.generated.h"
+
+UCLASS()
+class URichCurveKeyProxy : public UObject, public ICurveEditorKeyProxy
+{
+public:
+	GENERATED_BODY()
+
+	/**
+	 * Initialize this key proxy object by caching the underlying key object, and retrieving the time/value each tick
+	 */
+	void Initialize(FKeyHandle InKeyHandle, FRichCurve* InRichCurve, TWeakObjectPtr<UObject> InWeakOwner)
+	{
+		KeyHandle  = InKeyHandle;
+		RichCurve  = InRichCurve;
+		WeakOwner  = InWeakOwner;
+
+		UObject* Owner = WeakOwner.Get();
+		if (Owner && RichCurve->IsKeyHandleValid(KeyHandle))
+		{
+			Value = RichCurve->GetKey(KeyHandle);
+		}
+	}
+
+private:
+
+	virtual void PostEditChangeProperty(FPropertyChangedEvent& PropertyChangedEvent) override
+	{
+		Super::PostEditChangeProperty(PropertyChangedEvent);
+
+		UObject* Owner = WeakOwner.Get();
+		if (Owner && RichCurve->IsKeyHandleValid(KeyHandle))
+		{
+			Owner->Modify();
+
+			FRichCurveKey& ActualKey = RichCurve->GetKey(KeyHandle);
+
+			const float PreviousTime = ActualKey.Time;
+			const float NewTime      = Value.Time;
+
+			// Copy the key properties
+			ActualKey = Value;
+			ActualKey.Time = PreviousTime;
+			if (PreviousTime != NewTime)
+			{
+				RichCurve->SetKeyTime(KeyHandle, NewTime);
+			}
+		}
+	}
+
+	virtual void UpdateValuesFromRawData() override
+	{
+		UObject* Owner = WeakOwner.Get();
+		if (Owner && RichCurve->IsKeyHandleValid(KeyHandle))
+		{
+			Value = RichCurve->GetKey(KeyHandle);
+		}
+	}
+
+private:
+
+	/** User-facing value of the key, applied to the actual key on PostEditChange, and updated every tick */
+	UPROPERTY(EditAnywhere, Category="Key", meta=(ShowOnlyInnerProperties))
+	FRichCurveKey Value;
+
+private:
+
+	/** Cached key handle that this key proxy relates to */
+	FKeyHandle KeyHandle;
+	/** Cached curve in which the key resides */
+	FRichCurve* RichCurve;
+	/** Cached owner in which the raw curve resides */
+	TWeakObjectPtr<UObject> WeakOwner;
+};
\ No newline at end of file
