This is a significant overhaul to the Curve Editor used by Sequencer which adds a plugin-based architecture and extensibility. New tools and toolbar buttons can be added to all usages of the curve edi

diff --git a/Engine/Source/Editor/CurveEditor/Private/SCurveEditorPanel.cpp b/Engine/Source/Editor/CurveEditor/Private/SCurveEditorPanel.cpp
--- a/Engine/Source/Editor/CurveEditor/Private/SCurveEditorPanel.cpp
+++ b/Engine/Source/Editor/CurveEditor/Private/SCurveEditorPanel.cpp
@@ -2,71 +2,117 @@
 
 #include "SCurveEditorPanel.h"
 #include "Templates/Tuple.h"
+#include "Algo/Partition.h"
 #include "Rendering/DrawElements.h"
 #include "CurveDrawInfo.h"
 #include "CurveEditorSettings.h"
-#include "CurveEditorContextMenu.h"
-#include "DragOperations/CurveEditorDragOperation_Tangent.h"
-#include "DragOperations/CurveEditorDragOperation_MoveKeys.h"
-#include "DragOperations/CurveEditorDragOperation_Pan.h"
-#include "DragOperations/CurveEditorDragOperation_Zoom.h"
-#include "DragOperations/CurveEditorDragOperation_Marquee.h"
 #include "CurveEditorCommands.h"
 #include "Framework/Commands/UICommandList.h"
 #include "Framework/MultiBox/MultiBoxBuilder.h"
-#include "Widgets/SToolTip.h"
 #include "Widgets/SNullWidget.h"
+#include "Widgets/SOverlay.h"
 #include "Widgets/Layout/SScrollBox.h"
 #include "Widgets/Layout/SSplitter.h"
 #include "Widgets/Text/STextBlock.h"
 #include "EditorStyleSet.h"
 #include "Misc/Attribute.h"
 #include "Algo/Sort.h"
-#include "Widgets/Input/SNumericEntryBox.h"
-#include "Widgets/Input/STextEntryPopup.h"
 #include "CurveEditorEditObjectContainer.h"
 #include "Modules/ModuleManager.h"
 #include "PropertyEditorModule.h"
 #include "IDetailsView.h"
+#include "IPropertyRowGenerator.h"
+#include "SCurveKeyDetailPanel.h"
+#include "ICurveEditorExtension.h"
+#include "ISequencerWidgetsModule.h"
+#include "CurveEditorScreenSpace.h"
+#include "Widgets/Layout/SBox.h"
+#include "Framework/MultiBox/MultiBoxExtender.h"
+#include "Filters/SCurveEditorFilterPanel.h"
+#include "Filters/CurveEditorFilterBase.h"
+#include "Filters/CurveEditorBakeFilter.h"
+#include "Filters/CurveEditorReduceFilter.h"
+#include "Widgets/Input/SNumericList.h"
+#include "Widgets/SFrameRatePicker.h"
+#include "CommonFrameRates.h"
+#include "CurveEditorViewRegistry.h"
+#include "SCurveEditorViewContainer.h"
+#include "Fonts/FontMeasure.h"
 
 #define LOCTEXT_NAMESPACE "SCurveEditorPanel"
 
-namespace
+int32 GCurveEditorPinnedViews = 0;
+static FAutoConsoleVariableRef CVarCurveEditorPinnedViews(
+	TEXT("CurveEditor.PinnedViews"),
+	GCurveEditorPinnedViews,
+	TEXT("Whether pinning a curve should also cause it to be exclusively added to a pinned view or not (default: off), rather than simply always remain visible.")
+);
+
+int32 GCurveEditorMaxCurvesPerPinnedView = 0;
+static FAutoConsoleVariableRef CVarCurveEditorMaxCurvesPerPinnedView(
+	TEXT("CurveEditor.MaxCurvesPerPinnedView"),
+	GCurveEditorMaxCurvesPerPinnedView,
+	TEXT("When CurveEditor.PinnedViews is 1, defines the maximum number of curves allowed on a pinned view (0 for no maximum).")
+);
+
+/**
+ * Utility class to hold an overlay which tells the user how to use the editor when there are no curves selected.
+ */
+class SCurveEditorViewOverlay : public SCompoundWidget
 {
-	static float HoverProximityThresholdPx = 5.f;
-	static float HoveredCurveThickness = 5.f;
-	static float UnHoveredCurveThickness = 1.f;
-	static float LabelOffsetPx = 2.f;
-	static bool  bAntiAliasCurves = true;
+	SLATE_BEGIN_ARGS(SCurveEditorViewOverlay)
+	{}
+	SLATE_END_ARGS()
 
-	FReply HandledReply(const FGeometry&, const FPointerEvent&)
-	{
-		return FReply::Handled();
-	}
-}
+	void Construct(const FArguments& InArgs) {}
 
 
-TUniquePtr<ICurveEditorKeyDragOperation> CreateKeyDrag(ECurvePointType KeyType)
-{
-	switch (KeyType)
+	virtual int32 OnPaint(const FPaintArgs& Args, const FGeometry& AllottedGeometry, const FSlateRect& MyCullingRect, FSlateWindowElementList& OutDrawElements, int32 LayerId, const FWidgetStyle& InWidgetStyle, bool bParentEnabled) const override
 	{
-	case ECurvePointType::ArriveTangent:
-	case ECurvePointType::LeaveTangent:
-		return MakeUnique<FCurveEditorDragOperation_Tangent>();
+		static const FLinearColor BackgroundColor = FLinearColor::Black.CopyWithNewOpacity(0.35f);
+		static const FText InstructionText = LOCTEXT("CurveEditorTutorialOverlay", "Select a curve on the left to begin editing.");
+		const FSlateBrush*   WhiteBrush = FEditorStyle::GetBrush("WhiteBrush");
+		const TSharedRef<FSlateFontMeasure> FontMeasure = FSlateApplication::Get().GetRenderer()->GetFontMeasureService();
+		const FSlateFontInfo FontInfo = FCoreStyle::Get().GetFontStyle("FontAwesome.13");
+
+		// Draw a darkened background
+		{
+			FSlateDrawElement::MakeBox(OutDrawElements, LayerId, AllottedGeometry.ToPaintGeometry(), WhiteBrush, ESlateDrawEffect::None, BackgroundColor);
+		}
 
-	default:
-		return MakeUnique<FCurveEditorDragOperation_MoveKeys>();
+		// Draw the tutorial text centered
+		{
+			const FVector2D LabelSize = FontMeasure->Measure(InstructionText, FontInfo);
+			const FVector2D GeometrySize = AllottedGeometry.GetLocalSize();
+			const FVector2D LabelOffset = FVector2D((GeometrySize.X - LabelSize.X) / 2.f, (GeometrySize.Y - LabelSize.Y) / 2.f);
+			const FPaintGeometry LabelGeometry = AllottedGeometry.ToPaintGeometry(FSlateLayoutTransform(LabelOffset));
+
+			FSlateDrawElement::MakeText(OutDrawElements, LayerId, LabelGeometry, InstructionText, FontInfo);
+		}
+
+		return LayerId + 1;
 	}
-}
+};
 
-class SDynamicToolTip : public SToolTip
+/**
+ * Implemented as a friend struct to SCurveEditorView to ensure that SCurveEditorPanel is the only thing that can add/remove curves from views
+ * whilst disallowing access to any other private members;
+ */
+struct FCurveEditorPanelViewTracker
 {
-public:
-	TAttribute<bool> bIsEnabled;
-	virtual bool IsEmpty() const override { return !bIsEnabled.Get(); }
+	static void AddCurveToView(SCurveEditorView* View, FCurveModelID InCurveID)
+	{
+		View->AddCurve(InCurveID);
+	}
+	static void RemoveCurveFromView(SCurveEditorView* View, FCurveModelID InCurveID)
+	{
+		View->RemoveCurve(InCurveID);
+	}
 };
 
 SCurveEditorPanel::SCurveEditorPanel()
+	: bNeedsRefresh(true)
+	, CachedActiveCurvesSerialNumber(-1)
 {
 	EditObjects = MakeUnique<FCurveEditorEditObjectContainer>();
 	bSelectionSupportsWeightedTangents = false;
@@ -74,106 +120,143 @@ SCurveEditorPanel::SCurveEditorPanel()
 
 SCurveEditorPanel::~SCurveEditorPanel()
 {
+	// Attempt to close a dialog if it's open. It has a weak reference to us and doesn't work well when it's invalid.
+	SCurveEditorFilterPanel::CloseDialog();
 }
 
 void SCurveEditorPanel::Construct(const FArguments& InArgs, TSharedRef<FCurveEditor> InCurveEditor)
 {
 	GridLineTintAttribute = InArgs._GridLineTint;
+	DisabledTimeSnapTooltipAttribute = InArgs._DisabledTimeSnapTooltip;
+	WeakTabManager = InArgs._TabManager;
 
 	CachedSelectionSerialNumber = 0;
 	CurveEditor = InCurveEditor;
 
-	ReduceTolerance = 0.1f;
+	CurveEditor->SetPanel(SharedThis(this));
 
 	CurveEditor->BindCommands();
+	CurveEditor->SetTimeSliderController(InArgs._ExternalTimeSliderController);
 
 	CommandList = MakeShared<FUICommandList>();
 	CommandList->Append(InCurveEditor->GetCommands().ToSharedRef());
 
 	BindCommands();
-	SetClipping(EWidgetClipping::ClipToBounds);
 
+	// Create some Widgets
+	ISequencerWidgetsModule& SequencerWidgets = FModuleManager::Get().LoadModuleChecked<ISequencerWidgetsModule>("SequencerWidgets");
+	TSharedPtr<SWidget> TopTimeSlider = SNullWidget::NullWidget;
+	if (InArgs._ExternalTimeSliderController)
 	{
-		FPropertyEditorModule& EditModule = FModuleManager::Get().GetModuleChecked<FPropertyEditorModule>("PropertyEditor");
-		FDetailsViewArgs DetailsViewArgs(
-			/*bUpdateFromSelection=*/ false,
-			/*bLockable=*/ false,
-			/*bAllowSearch=*/ false,
-			FDetailsViewArgs::HideNameArea,
-			/*bHideSelectionTip=*/ true,
-			/*InNotifyHook=*/ nullptr,
-			/*InSearchInitialKeyFocus=*/ false,
-			/*InViewIdentifier=*/ NAME_None);
-		DetailsViewArgs.DefaultsOnlyVisibility = EEditDefaultsOnlyNodeVisibility::Automatic;
-		DetailsViewArgs.bShowOptions = false;
-
-		KeyDetailsView = EditModule.CreateDetailView(DetailsViewArgs);
+		TopTimeSlider = SequencerWidgets.CreateTimeSlider(InArgs._ExternalTimeSliderController.ToSharedRef(), false /*bMirrorLabels*/);
 	}
 
-	ChildSlot
-	[
+	TSharedRef<SScrollBar> ScrollBar = SNew(SScrollBar).Thickness(FVector2D(5.f, 5.f));
+
+	TSharedRef<SWidget> MainContent = 
 		SNew(SOverlay)
 
+		// The main editing area
 		+ SOverlay::Slot()
-		.VAlign(VAlign_Top)
 		[
-			SNew(SSplitter)
-			.Visibility(this, &SCurveEditorPanel::GetSplitterVisibility)
+			SNew(SVerticalBox)
 
-			+ SSplitter::Slot()
-			.Value(0.25f)
+			+ SVerticalBox::Slot()
+			.AutoHeight()
 			[
+				// Top Time Slider
 				SNew(SBorder)
-				.Padding(FMargin(0))
-				.BorderImage(FEditorStyle::GetBrush("NoBorder"))
-				.OnMouseButtonDown_Static(HandledReply)
-				.OnMouseMove_Static(HandledReply)
-				.OnMouseButtonUp_Static(HandledReply)
+				.BorderImage(FEditorStyle::GetBrush("ToolPanel.GroupBorder"))
+				.BorderBackgroundColor(FLinearColor(.50f, .50f, .50f, 1.0f))
+				.Padding(0)
+				.Clipping(EWidgetClipping::ClipToBounds)
 				[
-					KeyDetailsView.ToSharedRef()
+					TopTimeSlider.ToSharedRef()
 				]
 			]
 
-			+ SSplitter::Slot()
-			.Value(0.75f)
+			+ SVerticalBox::Slot()
 			[
-				SNullWidget::NullWidget
+				SNew(SOverlay)
+
+				+ SOverlay::Slot()
+				[
+					SAssignNew(ScrollBox, SScrollBox)
+					.ExternalScrollbar(ScrollBar)
+
+					+ SScrollBox::Slot()
+					[
+						// Main Curve View Area. The contents of this are dynamically filled based
+						// on your current views.
+
+						SAssignNew(CurveViewsContainer, SCurveEditorViewContainer, InCurveEditor)
+						.ExternalTimeSliderController(InArgs._ExternalTimeSliderController)
+					]
+				]
+
+				+ SOverlay::Slot()
+				.HAlign(HAlign_Right)
+				[
+					ScrollBar
+				]
 			]
 		]
-	];
 
-	UpdateEditBox();
+		// An overlay for the main area which lets us put system-wide overlays
+		+ SOverlay::Slot()
+		[
+			SNew(SCurveEditorViewOverlay)
+			.Visibility(this, &SCurveEditorPanel::ShouldInstructionOverlayBeVisible)
+		];
+	
 
-	TSharedRef<SDynamicToolTip> ToolTipWidget =
-		SNew(SDynamicToolTip)
-		.BorderImage(FCoreStyle::Get().GetBrush("ToolTip.BrightBackground"))
+	if (InArgs._TreeContent.Widget != SNullWidget::NullWidget)
+	{
+		ChildSlot
 		[
-			SNew(SVerticalBox)
-			+ SVerticalBox::Slot()
-			[
-				SNew(STextBlock)
-				.Text(this, &SCurveEditorPanel::GetToolTipCurveName)
-				.Font(FCoreStyle::Get().GetFontStyle("ToolTip.LargerFont"))
-				.ColorAndOpacity(FLinearColor::Black)
-			]
-			+ SVerticalBox::Slot()
+			SNew(SSplitter)
+			.Orientation(Orient_Horizontal)
+
+			+ SSplitter::Slot()
+			.Value(0.25f)
 			[
-				SNew(STextBlock)
-				.Text(this, &SCurveEditorPanel::GetToolTipTimeText)
-				.Font(FCoreStyle::Get().GetFontStyle("ToolTip.LargerFont"))
-				.ColorAndOpacity(FLinearColor::Black)
+				InArgs._TreeContent.Widget
 			]
-			+ SVerticalBox::Slot()
+
+			+ SSplitter::Slot()
+			.Value(0.75f)
 			[
-				SNew(STextBlock)
-				.Text(this, &SCurveEditorPanel::GetToolTipValueText)
-				.Font(FCoreStyle::Get().GetFontStyle("ToolTip.LargerFont"))
-				.ColorAndOpacity(FLinearColor::Black)
+				MainContent
 			]
 		];
+	}
+	else
+	{
+		ChildSlot
+		[
+			MainContent
+		];
+	}
 
-	ToolTipWidget->bIsEnabled = MakeAttributeSP(this, &SCurveEditorPanel::IsToolTipEnabled);
-	SetToolTip(ToolTipWidget);
+
+
+	KeyDetailsView = SNew(SCurveKeyDetailPanel, CurveEditor.ToSharedRef())
+		.IsEnabled(this, &SCurveEditorPanel::IsInlineEditPanelEditable);
+
+	UpdateEditBox();
+
+	// Initializes our Curve Views on the next Tick
+	SetViewMode(ECurveEditorViewID::Absolute);
+}
+
+TArrayView<const TSharedPtr<SCurveEditorView>> SCurveEditorPanel::GetViews() const
+{
+	return CurveViewsContainer->GetViews();
+}
+
+void SCurveEditorPanel::ScrollBy(float Amount)
+{
+	ScrollBox->SetScrollOffset(ScrollBox->GetScrollOffset() + Amount);
 }
 
 void SCurveEditorPanel::BindCommands()
@@ -196,11 +279,13 @@ void SCurveEditorPanel::BindCommands()
 		FIsActionChecked IsCubicBreakCommon = FIsActionChecked::CreateSP(this, &SCurveEditorPanel::CompareCommonTangentMode, RCIM_Cubic, RCTM_Break);
 		FIsActionChecked IsCubicWeightCommon = FIsActionChecked::CreateSP(this, &SCurveEditorPanel::CompareCommonTangentWeightMode, RCIM_Cubic, RCTWM_WeightedBoth);
 
-		CommandList->MapAction(FCurveEditorCommands::Get().InterpolationConstant, SetConstant, FCanExecuteAction(), IsConstantCommon);
-		CommandList->MapAction(FCurveEditorCommands::Get().InterpolationLinear, SetLinear, FCanExecuteAction(), IsLinearCommon);
-		CommandList->MapAction(FCurveEditorCommands::Get().InterpolationCubicAuto, SetCubicAuto, FCanExecuteAction(), IsCubicAutoCommon);
-		CommandList->MapAction(FCurveEditorCommands::Get().InterpolationCubicUser, SetCubicUser, FCanExecuteAction(), IsCubicUserCommon);
-		CommandList->MapAction(FCurveEditorCommands::Get().InterpolationCubicBreak, SetCubicBreak, FCanExecuteAction(), IsCubicBreakCommon);
+		FCanExecuteAction CanSetKeyTangent = FIsActionChecked::CreateSP(this, &SCurveEditorPanel::CanSetKeyInterpolation);
+
+		CommandList->MapAction(FCurveEditorCommands::Get().InterpolationConstant, SetConstant, CanSetKeyTangent, IsConstantCommon);
+		CommandList->MapAction(FCurveEditorCommands::Get().InterpolationLinear, SetLinear, CanSetKeyTangent, IsLinearCommon);
+		CommandList->MapAction(FCurveEditorCommands::Get().InterpolationCubicAuto, SetCubicAuto, CanSetKeyTangent, IsCubicAutoCommon);
+		CommandList->MapAction(FCurveEditorCommands::Get().InterpolationCubicUser, SetCubicUser, CanSetKeyTangent, IsCubicUserCommon);
+		CommandList->MapAction(FCurveEditorCommands::Get().InterpolationCubicBreak, SetCubicBreak, CanSetKeyTangent, IsCubicBreakCommon);
 		CommandList->MapAction(FCurveEditorCommands::Get().InterpolationToggleWeighted, ToggleWeighted, CanToggleWeighted, IsCubicWeightCommon);
 
 	}
@@ -246,41 +331,292 @@ void SCurveEditorPanel::BindCommands()
 		CommandList->MapAction(FCurveEditorCommands::Get().SetPostInfinityExtrapLinear, SetLinear, FCanExecuteAction(), IsLinearCommon);
 		CommandList->MapAction(FCurveEditorCommands::Get().SetPostInfinityExtrapConstant, SetConstant, FCanExecuteAction(), IsConstantCommon);
 	}
-	//override for key reduction so we show popup
+
+	// Absolute, Stacked and Normalized views.
+	{
+		FExecuteAction SetViewAbsolute = FExecuteAction::CreateSP(this, &SCurveEditorPanel::SetViewMode, ECurveEditorViewID::Absolute);
+		FIsActionChecked IsViewModeAbsolute = FIsActionChecked::CreateSP(this, &SCurveEditorPanel::CompareViewMode, ECurveEditorViewID::Absolute);
+
+		FExecuteAction SetViewStacked = FExecuteAction::CreateSP(this, &SCurveEditorPanel::SetViewMode, ECurveEditorViewID::Stacked);
+		FIsActionChecked IsViewModeStacked = FIsActionChecked::CreateSP(this, &SCurveEditorPanel::CompareViewMode, ECurveEditorViewID::Stacked);
+
+		FExecuteAction SetViewNormalized = FExecuteAction::CreateSP(this, &SCurveEditorPanel::SetViewMode, ECurveEditorViewID::Normalized);
+		FIsActionChecked IsViewModeNormalized = FIsActionChecked::CreateSP(this, &SCurveEditorPanel::CompareViewMode, ECurveEditorViewID::Normalized);
+
+		CommandList->MapAction(FCurveEditorCommands::Get().SetViewModeAbsolute, SetViewAbsolute, FCanExecuteAction(), IsViewModeAbsolute);
+		CommandList->MapAction(FCurveEditorCommands::Get().SetViewModeStacked, SetViewStacked, FCanExecuteAction(), IsViewModeStacked);
+		CommandList->MapAction(FCurveEditorCommands::Get().SetViewModeNormalized, SetViewNormalized, FCanExecuteAction(), IsViewModeNormalized);
+	}
+
 	{
-		CommandList->MapAction(FCurveEditorCommands::Get().ReduceCurve, FExecuteAction::CreateSP(this, &SCurveEditorPanel::OnSimplifySelection));
+		// Deselect Current Keys
+		TSharedPtr<FCurveEditor> LocalCurveEditor = CurveEditor;
+		FExecuteAction DeselectAllAction = FExecuteAction::CreateLambda([LocalCurveEditor] { if (LocalCurveEditor.IsValid())
+			{
+				LocalCurveEditor->GetSelection().Clear();
+			} 
+		});
+		CommandList->MapAction(FCurveEditorCommands::Get().DeselectAllKeys, DeselectAllAction);
 	}
+
+	// Presets for Bake and Reduce 
+	CommandList->MapAction(FCurveEditorCommands::Get().BakeCurve, FExecuteAction::CreateSP(this, &SCurveEditorPanel::ShowCurveFilterUI, TSubclassOf<UCurveEditorFilterBase>(UCurveEditorBakeFilter::StaticClass())));
+	CommandList->MapAction(FCurveEditorCommands::Get().ReduceCurve, FExecuteAction::CreateSP(this, &SCurveEditorPanel::ShowCurveFilterUI, TSubclassOf<UCurveEditorFilterBase>(UCurveEditorReduceFilter::StaticClass())));
+
+	// User Implementable Filter  just defaults to Bake since we know it exists...
+	CommandList->MapAction(FCurveEditorCommands::Get().OpenUserImplementableFilterWindow, FExecuteAction::CreateSP(this, &SCurveEditorPanel::ShowCurveFilterUI, TSubclassOf<UCurveEditorFilterBase>(UCurveEditorBakeFilter::StaticClass())));
+
+	// Axis Snapping
+	{
+		FUIAction SetSnappingNoneAction(FExecuteAction::CreateSP(this, &SCurveEditorPanel::SetAxisSnapping, EAxisList::Type::None));
+		FUIAction SetSnappingHorizontalAction(FExecuteAction::CreateSP(this, &SCurveEditorPanel::SetAxisSnapping, EAxisList::Type::X));
+		FUIAction SetSnappingVerticalAction(FExecuteAction::CreateSP(this, &SCurveEditorPanel::SetAxisSnapping, EAxisList::Type::Y));
+
+		CommandList->MapAction(FCurveEditorCommands::Get().SetAxisSnappingNone, SetSnappingNoneAction);
+		CommandList->MapAction(FCurveEditorCommands::Get().SetAxisSnappingHorizontal, SetSnappingHorizontalAction);
+		CommandList->MapAction(FCurveEditorCommands::Get().SetAxisSnappingVertical, SetSnappingVerticalAction);
+	}
+
+}
+
+void SCurveEditorPanel::SetViewMode(const ECurveEditorViewID NewViewMode)
+{
+	DefaultViewID = NewViewMode;
+	bNeedsRefresh = true;
+}
+
+bool SCurveEditorPanel::CompareViewMode(const ECurveEditorViewID InViewMode) const
+{
+	return DefaultViewID == InViewMode;
+}
+
+void SCurveEditorPanel::SetAxisSnapping(EAxisList::Type InAxis)
+{
+	FCurveEditorAxisSnap Snap = CurveEditor->GetAxisSnap();
+	Snap.RestrictedAxisList = InAxis;
+	CurveEditor->SetAxisSnap(Snap);
 }
 
 void SCurveEditorPanel::Tick(const FGeometry& AllottedGeometry, const double InCurrentTime, const float InDeltaTime)
 {
+	bool bWasRefreshed = false;
+	if (bNeedsRefresh || CachedActiveCurvesSerialNumber != CurveEditor->GetActiveCurvesSerialNumber())
+	{
+		RebuildCurveViews();
+
+		if (CurveEditor->ShouldAutoFrame())
+		{
+			CurveEditor->ZoomToFitCurves(CurveEditor->GetEditedCurves().Array());
+		}
+
+		bNeedsRefresh = false;
+		CachedActiveCurvesSerialNumber = CurveEditor->GetActiveCurvesSerialNumber();
+		bWasRefreshed = true;
+	}
+
 	UpdateCommonCurveInfo();
 	UpdateEditBox();
 
-	CurveEditor->UpdateGeometry(AllottedGeometry);
+	CachedSelectionSerialNumber = CurveEditor->Selection.GetSerialNumber();
+}
 
-	CachedDrawParams.Reset();
-	CurveEditor->GetCurveDrawParams(CachedDrawParams);
+void SCurveEditorPanel::RemoveCurveFromViews(FCurveModelID InCurveID)
+{
+	for (auto It = CurveViews.CreateKeyIterator(InCurveID); It; ++It)
+	{
+		FCurveEditorPanelViewTracker::RemoveCurveFromView(&It.Value().Get(), InCurveID);
+		It.RemoveCurrent();
+	}
+}
 
-	CachedSelectionSerialNumber = CurveEditor->Selection.GetSerialNumber();
+void SCurveEditorPanel::AddView(TSharedRef<SCurveEditorView> ViewToAdd)
+{
+	ExternalViews.Add(ViewToAdd);
+	bNeedsRefresh = true;
+}
+
+void SCurveEditorPanel::RemoveView(TSharedRef<SCurveEditorView> ViewToRemove)
+{
+	ExternalViews.Remove(ViewToRemove);
+	bNeedsRefresh = true;
+}
+
+TSharedPtr<SCurveEditorView> SCurveEditorPanel::CreateViewOfType(FCurveModelID CurveModelID, ECurveEditorViewID ViewTypeID, bool bPinned)
+{
+	for (auto It = FreeViewsByType.CreateKeyIterator(ViewTypeID); It; ++It)
+	{
+		TSharedRef<SCurveEditorView> View = It.Value();
+
+		if (!GCurveEditorPinnedViews || View->bPinned == bPinned)
+		{
+			FCurveEditorPanelViewTracker::AddCurveToView(&View.Get(), CurveModelID);
+			CurveViews.Add(CurveModelID, View);
+
+			if (!View->HasCapacity())
+			{
+				It.RemoveCurrent();
+			}
+
+			return View;
+		}
+	}
+
+	TSharedPtr<SCurveEditorView> View = FCurveEditorViewRegistry::Get().ConstructView(ViewTypeID, CurveEditor);
+	if (View)
+	{
+		if (GCurveEditorPinnedViews && bPinned)
+		{
+			// Pinned views are always a fixed height
+			View->bPinned = true;
+			View->MaximumCapacity = View->MaximumCapacity == 0 ? GCurveEditorMaxCurvesPerPinnedView : FMath::Min(View->MaximumCapacity, GCurveEditorMaxCurvesPerPinnedView);
+			if (!View->FixedHeight.IsSet())
+			{
+				View->FixedHeight = 100.f;
+			}
+		}
+		View->ViewTypeID = ViewTypeID;
+		FCurveEditorPanelViewTracker::AddCurveToView(View.Get(), CurveModelID);
+		CurveViews.Add(CurveModelID, View.ToSharedRef());
+
+		if (View->HasCapacity())
+		{
+			FreeViewsByType.Add(ViewTypeID, View.ToSharedRef());
+		}
+	}
+
+	return View;
+}
+
+void SCurveEditorPanel::RebuildCurveViews()
+{
+	TSet<TSharedRef<SCurveEditorView>> Views = ExternalViews;
+
+	for (const TTuple<FCurveModelID, TUniquePtr<FCurveModel>>& CurvePair : CurveEditor->GetCurves())
+	{
+		FCurveModel* Curve = CurvePair.Value.Get();
+		FCurveModelID CurveID = CurvePair.Key;
+
+		const bool bIsPinned = CurveEditor->IsCurvePinned(CurveID);
+
+		bool bNeedsView = true;
+
+		for (auto ViewIt = CurveViews.CreateKeyIterator(CurveID); ViewIt; ++ViewIt)
+		{
+			TSharedRef<SCurveEditorView> View = ViewIt.Value();
+			// @todo: pinning - This code causes curves that have changed their pinned state to be re-added to a correctly (un)pinned views
+			if (GCurveEditorPinnedViews && View->bPinned != bIsPinned)
+			{
+				// No longer the same pinned status as the view it's in - remove it so that it can get added to the correct view (or removed entirely)
+				FCurveEditorPanelViewTracker::RemoveCurveFromView(&View.Get(), CurveID);
+				ViewIt.RemoveCurrent();
+			}
+			else if (View->ViewTypeID == DefaultViewID || !EnumHasAnyFlags(View->ViewTypeID, ECurveEditorViewID::ANY_BUILT_IN))
+			{
+				// Keep this view if it is the default view or any other custom view
+				Views.Add(View);
+				bNeedsView = false;
+			}
+			else
+			{
+				// Built in view which is no longer the selected mode - remove it
+				FCurveEditorPanelViewTracker::RemoveCurveFromView(&View.Get(), CurveID);
+				ViewIt.RemoveCurrent();
+			}
+		}
+
+		if (bNeedsView)
+		{
+			ECurveEditorViewID SupportedViews = Curve->GetSupportedViews();
+
+			// Add to the default view if supported, else use the first supported view we can find.
+			// @todo: this may require extra work if curves are ever to support multiple views but it's fine for now
+			if (EnumHasAnyFlags(SupportedViews, DefaultViewID))
+			{
+				TSharedPtr<SCurveEditorView> NewView = CreateViewOfType(CurveID, DefaultViewID, bIsPinned);
+				if (NewView)
+				{
+					Views.Add(NewView.ToSharedRef());
+				}
+				continue;
+			}
+
+			ECurveEditorViewID CustomView = ECurveEditorViewID::CUSTOM_START;
+			while (CustomView >= ECurveEditorViewID::CUSTOM_START)
+			{
+				if (EnumHasAnyFlags(SupportedViews, ECurveEditorViewID(CustomView)))
+				{
+					TSharedPtr<SCurveEditorView> NewView = CreateViewOfType(CurveID, CustomView, bIsPinned);
+					if (NewView)
+					{
+						Views.Add(NewView.ToSharedRef());
+					}
+				}
+				CustomView = ECurveEditorViewID((__underlying_type(ECurveEditorViewID))(CustomView) << 1);
+			}
+		}
+	}
+
+	// Remove any empty views
+	for (auto It = FreeViewsByType.CreateIterator(); It; ++It)
+	{
+		if (!It.Value()->bAllowEmpty && It.Value()->NumCurves() == 0)
+		{
+			It.RemoveCurrent();
+		}
+	}
+
+	// Sort by pinned, then capacity
+	TArray<TSharedRef<SCurveEditorView>> SortedViews = Views.Array();
+	Algo::Sort(SortedViews, [](const TSharedPtr<SCurveEditorView>& A, const TSharedPtr<SCurveEditorView>& B)
+		{
+			if (A->SortBias == B->SortBias)
+			{
+				if (A->bPinned == B->bPinned)
+				{
+					return A->RelativeOrder < B->RelativeOrder;
+				}
+
+				return A->bPinned != 0;
+			}
+
+			return A->SortBias > B->SortBias;
+		}
+	);
+
+	CurveViewsContainer->Clear();
+	for (TSharedPtr<SCurveEditorView> View : SortedViews)
+	{
+		CurveViewsContainer->AddView(View.ToSharedRef());
+	}
 }
 
 void SCurveEditorPanel::UpdateCommonCurveInfo()
 {
 	// Gather up common extended curve info for the current set of curves
 	TOptional<FCurveAttributes> AccumulatedCurveAttributes;
-	for (const TTuple<FCurveModelID, TUniquePtr<FCurveModel>>& Pair : CurveEditor->GetCurves())
+	for (const FCurveModelID& CurveID : CurveEditor->GetEditedCurves())
 	{
 		FCurveAttributes Attributes;
-		Pair.Value->GetCurveAttributes(Attributes);
-
-		if (!AccumulatedCurveAttributes.IsSet())
-		{
-			AccumulatedCurveAttributes = Attributes;
-		}
-		else
+		
+		FCurveModel* Curve = CurveEditor->FindCurve(CurveID);
+		if (Curve)
 		{
-			AccumulatedCurveAttributes = FCurveAttributes::MaskCommon(AccumulatedCurveAttributes.GetValue(), Attributes);
+			Curve->GetCurveAttributes(Attributes);
+
+			// Some curves don't support extrapolation. We don't count them for determinine the accumulated state.
+			if (Attributes.HasPreExtrapolation() && Attributes.GetPreExtrapolation() == RCCE_None && Attributes.HasPostExtrapolation() && Attributes.GetPostExtrapolation() == RCCE_None)
+			{
+				continue;
+			}
+
+			if (!AccumulatedCurveAttributes.IsSet())
+			{
+				AccumulatedCurveAttributes = Attributes;
+			}
+			else
+			{
+				AccumulatedCurveAttributes = FCurveAttributes::MaskCommon(AccumulatedCurveAttributes.GetValue(), Attributes);
+			}
 		}
 	}
 
@@ -323,21 +659,6 @@ void SCurveEditorPanel::UpdateCommonCurveInfo()
 	CachedCommonKeyAttributes = AccumulatedKeyAttributes.Get(FKeyAttributes());
 }
 
-void SCurveEditorPanel::RebindContextualActions(FVector2D MousePosition)
-{
-	CommandList->UnmapAction(FCurveEditorCommands::Get().AddKeyHovered);
-	CommandList->UnmapAction(FCurveEditorCommands::Get().AddKeyToAllCurves);
-	CommandList->UnmapAction(FCurveEditorCommands::Get().AddKeyToAllCurvesHere);
-
-	TOptional<FCurveModelID> HoveredCurve = GetHoveredCurve();
-	if (HoveredCurve.IsSet())
-	{
-		CommandList->MapAction(FCurveEditorCommands::Get().AddKeyHovered, FExecuteAction::CreateSP(this, &SCurveEditorPanel::OnAddKey, HoveredCurve.GetValue(), MousePosition));
-	}
-	CommandList->MapAction(FCurveEditorCommands::Get().AddKeyToAllCurves, FExecuteAction::CreateSP(this, &SCurveEditorPanel::OnAddKey, MousePosition));
-	CommandList->MapAction(FCurveEditorCommands::Get().AddKeyToAllCurvesHere, FExecuteAction::CreateSP(this, &SCurveEditorPanel::OnAddKeyHere, MousePosition));
-}
-
 void SCurveEditorPanel::UpdateEditBox()
 {
 	const FCurveEditorSelection& Selection = CurveEditor->Selection;
@@ -398,70 +719,18 @@ void SCurveEditorPanel::UpdateEditBox()
 				{
 					KeyHandleToEditObject.Add(KeyHandleScratch[Index], NewObject);
 					AllEditObjects.Add(NewObject);
+
+					// Update the proxy immediately after adding it so that it doesn't have the wrong values for 1 tick.
+					if (ICurveEditorKeyProxy* Proxy = Cast<ICurveEditorKeyProxy>(NewObject))
+					{
+						Proxy->UpdateValuesFromRawData();
+					}
 				}
 			}
 		}
 	}
 
-	KeyDetailsView->SetObjects(AllEditObjects);
-}
-
-void SCurveEditorPanel::UpdateCurveProximities(FVector2D MousePixel)
-{
-	CachedToolTipData.Reset();
-
-	// Update curve proximities
-	CurveProximities.Reset();
-
-	FCurveEditorScreenSpace ScreenSpace = CurveEditor->GetScreenSpace();
-
-	double MinMouseTime    = ScreenSpace.ScreenToSeconds(MousePixel.X - HoverProximityThresholdPx);
-	double MaxMouseTime    = ScreenSpace.ScreenToSeconds(MousePixel.X + HoverProximityThresholdPx);
-	double MouseValue      = ScreenSpace.ScreenToValue(MousePixel.Y);
-	float  PixelsPerOutput = ScreenSpace.PixelsPerOutput();
-
-	FVector2D MinPos(MousePixel.X - HoverProximityThresholdPx, 0.0f);
-	FVector2D MaxPos(MousePixel.X + HoverProximityThresholdPx, 0.0f);
-
-	for (const TTuple<FCurveModelID, TUniquePtr<FCurveModel>>& Pair : CurveEditor->GetCurves())
-	{
-		double InputOffset = Pair.Value->GetInputDisplayOffset();
-		double MinEvalTime = MinMouseTime - InputOffset;
-		double MaxEvalTime = MaxMouseTime - InputOffset;
-
-		double MinValue = 0.0, MaxValue = 0.0;
-		if (Pair.Value->Evaluate(MinEvalTime, MinValue) && Pair.Value->Evaluate(MaxEvalTime, MaxValue))
-		{
-			MinPos.Y = ScreenSpace.ValueToScreen(MinValue);
-			MaxPos.Y = ScreenSpace.ValueToScreen(MaxValue);
-
-			float Distance = (FMath::ClosestPointOnSegment2D(MousePixel, MinPos, MaxPos) - MousePixel).Size();
-			CurveProximities.Add(MakeTuple(Pair.Key, Distance));
-		}
-	}
-
-	Algo::SortBy(CurveProximities, [](TTuple<FCurveModelID, float> In) { return In.Get<1>(); });
-
-	if (CurveProximities.Num() > 0 && CurveProximities[0].Get<1>() < HoverProximityThresholdPx)
-	{
-		const FCurveModel* HoveredCurve = CurveEditor->FindCurve(CurveProximities[0].Get<0>());
-
-		if (HoveredCurve)
-		{
-			double MouseTime = ScreenSpace.ScreenToSeconds(MousePixel.X) - HoveredCurve->GetInputDisplayOffset();
-			double EvaluatedTime = CurveEditor->GetSnapMetrics().SnapInputSeconds(MouseTime);
-
-			double EvaluatedValue = 0.0;
-			HoveredCurve->Evaluate(EvaluatedTime, EvaluatedValue);
-
-			FCachedToolTipData ToolTipData;
-			ToolTipData.Text           = HoveredCurve->GetDisplayName();
-			ToolTipData.EvaluatedTime  = FText::Format(NSLOCTEXT("CurveEditor", "CurveEditorTime", "{0}"), EvaluatedTime);
-			ToolTipData.EvaluatedValue = FText::Format(NSLOCTEXT("CurveEditor", "CurveEditorValue", "{0}"), EvaluatedValue);
-
-			CachedToolTipData = ToolTipData;
-		}
-	}
+	KeyDetailsView->GetPropertyRowGenerator()->SetObjects(AllEditObjects);
 }
 
 EVisibility SCurveEditorPanel::GetSplitterVisibility() const
@@ -469,36 +738,6 @@ EVisibility SCurveEditorPanel::GetSplitterVisibility() const
 	return EVisibility::Visible;
 }
 
-bool SCurveEditorPanel::IsToolTipEnabled() const
-{
-	return (CachedToolTipData.IsSet() && CurveEditor->GetSettings()->GetShowCurveEditorCurveToolTips());
-}
-
-FText SCurveEditorPanel::GetToolTipCurveName() const
-{
-	return CachedToolTipData.IsSet() ? CachedToolTipData->Text : FText();
-}
-
-FText SCurveEditorPanel::GetToolTipTimeText() const
-{
-	return CachedToolTipData.IsSet() ? CachedToolTipData->EvaluatedTime : FText();
-}
-
-FText SCurveEditorPanel::GetToolTipValueText() const
-{
-	return CachedToolTipData.IsSet() ? CachedToolTipData->EvaluatedValue : FText();
-}
-
-TOptional<FCurveModelID> SCurveEditorPanel::GetHoveredCurve() const
-{
-	if (CurveProximities.Num() > 0 && CurveProximities[0].Get<1>() < HoverProximityThresholdPx)
-	{
-		return CurveProximities[0].Get<0>();
-	}
-
-	return TOptional<FCurveModelID>();
-}
-
 void SCurveEditorPanel::SetKeyAttributes(FKeyAttributes KeyAttributes, FText Description)
 {
 	FScopedTransaction Transaction(Description);
@@ -530,7 +769,7 @@ void SCurveEditorPanel::ToggleWeightedTangents()
 
 	TMap<FCurveModelID, TArray<FKeyAttributes>> KeyAttributesPerCurve;
 
-	const TMap<FCurveModelID, FKeyHandleSet>& Selection = CurveEditor->Selection.GetAll();
+	const TMap<FCurveModelID, FKeyHandleSet>& Selection = CurveEditor->GetSelection().GetAll();
 
 	// Disable weights unless we find something that doesn't have weights, then add them
 	FKeyAttributes KeyAttributesToAssign = FKeyAttributes().SetTangentWeightMode(RCTWM_WeightedNone);
@@ -580,649 +819,432 @@ void SCurveEditorPanel::ToggleWeightedTangents()
 
 bool SCurveEditorPanel::CanToggleWeightedTangents() const
 {
-	return bSelectionSupportsWeightedTangents;
+	return bSelectionSupportsWeightedTangents && CanSetKeyInterpolation();
 }
 
-void SCurveEditorPanel::OnAddKey(FVector2D MousePixel)
+bool SCurveEditorPanel::CanSetKeyInterpolation() const
 {
-	FScopedTransaction Transcation(LOCTEXT("OnAddKey", "Add Key"));
-
-	bool bAddedKey = false;
-	FKeyAttributes DefaultAttributes = CurveEditor->DefaultKeyAttributes.Get();
-
-	double MouseTime = CurveEditor->GetScreenSpace().ScreenToSeconds(MousePixel.X);
-
-	CurveEditor->Selection.Clear();
-	for (const TTuple<FCurveModelID, TUniquePtr<FCurveModel>>& Pair : CurveEditor->GetCurves())
-	{
-		double EvalTime = MouseTime - Pair.Value->GetInputDisplayOffset();
-
-		double CurveValue = 0.0;
-		if (Pair.Value->Evaluate(EvalTime, CurveValue))
-		{
-			Pair.Value->Modify();
-
-			// Add a key on this curve
-			TOptional<FKeyHandle> NewKey = Pair.Value->AddKey(FKeyPosition(EvalTime, CurveValue), DefaultAttributes);
-			if (NewKey.IsSet())
-			{
-				bAddedKey = true;
-				CurveEditor->Selection.Add(FCurvePointHandle(Pair.Key, ECurvePointType::Key, NewKey.GetValue()));
-			}
-		}
-	}
-
-	if (!bAddedKey)
-	{
-		Transcation.Cancel();
-	}
+	return CurveEditor->GetSelection().Count() > 0;
 }
 
-void SCurveEditorPanel::OnAddKeyHere(FVector2D MousePixel)
+FReply SCurveEditorPanel::OnKeyDown(const FGeometry& MyGeometry, const FKeyEvent& InKeyEvent)
 {
-	FScopedTransaction Transcation(LOCTEXT("OnAddKeyHere", "Add Key"));
-
-	bool bAddedKey = false;
-
-	FKeyAttributes DefaultAttributes = CurveEditor->DefaultKeyAttributes.Get();
-
-	double MouseTime = CurveEditor->GetScreenSpace().ScreenToSeconds(MousePixel.X);
-	double NewValue  = CurveEditor->GetScreenSpace().ScreenToValue(MousePixel.Y);
-
-	CurveEditor->Selection.Clear();
-	for (const TTuple<FCurveModelID, TUniquePtr<FCurveModel>>& Pair : CurveEditor->GetCurves())
+	if (InKeyEvent.GetKey() == EKeys::Escape)
 	{
-		Pair.Value->Modify();
-
-		double KeyTime = MouseTime - Pair.Value->GetInputDisplayOffset();
-		TOptional<FKeyHandle> NewKey = Pair.Value->AddKey(FKeyPosition(KeyTime, NewValue), DefaultAttributes);
-		if (NewKey.IsSet())
-		{
-			bAddedKey = true;
-			CurveEditor->Selection.Add(FCurvePointHandle(Pair.Key, ECurvePointType::Key, NewKey.GetValue()));
-		}
+		CurveEditor->Selection.Clear();
+		return FReply::Handled();
 	}
-
-	if (!bAddedKey)
+	else if (CommandList->ProcessCommandBindings(InKeyEvent))
 	{
-		Transcation.Cancel();
+		return FReply::Handled();
 	}
+	return FReply::Unhandled();
 }
 
-void SCurveEditorPanel::OnAddKey(FCurveModelID CurveToAdd, FVector2D MousePixel)
+TSharedRef<SWidget> SCurveEditorPanel::MakeCurveEditorCurveViewOptionsMenu()
 {
-	FCurveModel* Curve = CurveEditor->FindCurve(CurveToAdd);
-	if (!Curve)
+	// This builds the dropdown menu when looking at the Curve View Options combobox.
+	FMenuBuilder MenuBuilder(true, CurveEditor->GetCommands());
+
+	MenuBuilder.BeginSection("TangentVisibility", LOCTEXT("CurveEditorMenuTangentVisibilityHeader", "Tangent Visibility"));
 	{
-		return;
+		MenuBuilder.AddMenuEntry(FCurveEditorCommands::Get().SetAllTangentsVisibility);
+		MenuBuilder.AddMenuEntry(FCurveEditorCommands::Get().SetSelectedKeysTangentVisibility);
+		MenuBuilder.AddMenuEntry(FCurveEditorCommands::Get().SetNoTangentsVisibility);
 	}
+	MenuBuilder.EndSection();
 
-	FScopedTransaction Transcation(FText::Format(LOCTEXT("OnKeyToCurveFormat", "Add Key to Curve '{0}'"), Curve->GetDisplayName()));
-
-	FKeyAttributes DefaultAttributes = CurveEditor->DefaultKeyAttributes.Get();
-
-	double MouseTime = CurveEditor->GetScreenSpace().ScreenToSeconds(MousePixel.X);
-	double NewValue  = CurveEditor->GetScreenSpace().ScreenToValue(MousePixel.Y);
+	MenuBuilder.AddMenuSeparator();
+	MenuBuilder.AddMenuEntry(FCurveEditorCommands::Get().ToggleAutoFrameCurveEditor);
+	MenuBuilder.AddMenuEntry(FCurveEditorCommands::Get().ToggleShowCurveEditorCurveToolTips);
 
-	CurveEditor->Selection.Clear();
-
-	Curve->Modify();
-
-	double KeyTime = MouseTime - Curve->GetInputDisplayOffset();
-	TOptional<FKeyHandle> NewKey = Curve->AddKey(FKeyPosition(KeyTime, NewValue), DefaultAttributes);
-	if (NewKey.IsSet())
-	{
-		CurveEditor->Selection.Add(FCurvePointHandle(CurveToAdd, ECurvePointType::Key, NewKey.GetValue()));
-	}
-	else
-	{
-		Transcation.Cancel();
-	}
+	return MenuBuilder.MakeWidget();
 }
 
-TOptional<FCurvePointHandle> SCurveEditorPanel::HitPoint(FVector2D MousePixel) const
+TSharedRef<SWidget> SCurveEditorPanel::MakeCurveExtrapolationMenu(const bool bInPostExtrapolation)
 {
-	TOptional<FCurvePointHandle> HitPoint;
-	TOptional<float> ClosestDistance;
+	// This builds the dropdown menu when looking at the Curve View Options combobox.
+	FMenuBuilder MenuBuilder(true, CurveEditor->GetCommands());
 
-	// Find all keys within the current hittest time
-	for (const FCurveDrawParams& Params : CachedDrawParams)
+	if (!bInPostExtrapolation)
 	{
-		for (const FCurvePointInfo& Point : Params.Points)
+		MenuBuilder.BeginSection("PreInfinity", LOCTEXT("CurveEditorMenuPreInfinityHeader", "Pre-Infinity"));
 		{
-			const FKeyDrawInfo& PointDrawInfo = Params.GetKeyDrawInfo(Point.Type);
-
-			FSlateRect KeyRect = FSlateRect::FromPointAndExtent(Point.ScreenPosition - PointDrawInfo.ScreenSize/2, PointDrawInfo.ScreenSize);
-			if (KeyRect.ContainsPoint(MousePixel))
-			{
-				float DistanceSquared = (KeyRect.GetCenter() - MousePixel).SizeSquared();
-				if (DistanceSquared <= ClosestDistance.Get(DistanceSquared))
-				{
-					ClosestDistance = DistanceSquared;
-					HitPoint = FCurvePointHandle(Params.GetID(), Point.Type, Point.KeyHandle);
-				}
-			}
+			MenuBuilder.AddMenuEntry(FCurveEditorCommands::Get().SetPreInfinityExtrapConstant);
+			MenuBuilder.AddMenuEntry(FCurveEditorCommands::Get().SetPreInfinityExtrapCycle);
+			MenuBuilder.AddMenuEntry(FCurveEditorCommands::Get().SetPreInfinityExtrapCycleWithOffset);
+			MenuBuilder.AddMenuEntry(FCurveEditorCommands::Get().SetPreInfinityExtrapLinear);
+			MenuBuilder.AddMenuEntry(FCurveEditorCommands::Get().SetPreInfinityExtrapOscillate);
 		}
+		MenuBuilder.EndSection();
 	}
-
-	return HitPoint;
-}
-
-FReply SCurveEditorPanel::OnMouseButtonDown(const FGeometry& MyGeometry, const FPointerEvent& MouseEvent)
-{
-	CachedToolTipData.Reset();
-
-	FCurveEditorScreenSpace ScreenSpace = CurveEditor->GetScreenSpace();
-
-	FVector2D MousePixel = MyGeometry.AbsoluteToLocal(MouseEvent.GetScreenSpacePosition());
-
-	// Mouse interaction that does not require a hit test
-	if (MouseEvent.GetEffectingButton() == EKeys::MiddleMouseButton)
+	else
 	{
-		if (MouseEvent.IsAltDown())
+		MenuBuilder.BeginSection("PostInfinity", LOCTEXT("CurveEditorMenuPostInfinityHeader", "Post-Infinity"));
 		{
-			DragOperation = FCurveEditorDelayedDrag(MousePixel, MouseEvent.GetEffectingButton());
-			DragOperation->DragImpl = MakeUnique<FCurveEditorDragOperation_Pan>(CurveEditor.Get());
-			return FReply::Handled().CaptureMouse(AsShared());
+			MenuBuilder.AddMenuEntry(FCurveEditorCommands::Get().SetPostInfinityExtrapConstant);
+			MenuBuilder.AddMenuEntry(FCurveEditorCommands::Get().SetPostInfinityExtrapCycle);
+			MenuBuilder.AddMenuEntry(FCurveEditorCommands::Get().SetPostInfinityExtrapCycleWithOffset);
+			MenuBuilder.AddMenuEntry(FCurveEditorCommands::Get().SetPostInfinityExtrapLinear);
+			MenuBuilder.AddMenuEntry(FCurveEditorCommands::Get().SetPostInfinityExtrapOscillate);
 		}
-		else
-		{
-			TOptional<FCurvePointHandle> NewPoint;
-
-			// Add a key to the closest curve to the mouse
-
-			if (TOptional<FCurveModelID> HoveredCurve = GetHoveredCurve())
-			{
-				FCurveModel* CurveToAddTo = CurveEditor->FindCurve(HoveredCurve.GetValue());
-				if (CurveToAddTo)
-				{
-					FScopedTransaction Transcation(LOCTEXT("InsertKey", "Insert Key"));
-
-					FKeyAttributes DefaultAttributes = CurveEditor->DefaultKeyAttributes.Get();
-
-					double MouseTime  = ScreenSpace.ScreenToSeconds(MousePixel.X);
-					double MouseValue = ScreenSpace.ScreenToValue(MousePixel.Y);
-
-					CurveToAddTo->Modify();
-
-					// Add a key on this curve
-					TOptional<FKeyHandle> NewKey = CurveToAddTo->AddKey(FKeyPosition(MouseTime, MouseValue), DefaultAttributes);
-					if (NewKey.IsSet())
-					{
-						NewPoint = FCurvePointHandle(HoveredCurve.GetValue(), ECurvePointType::Key, NewKey.GetValue());
-
-						CurveEditor->Selection.Clear();
-						CurveEditor->Selection.Add(NewPoint.GetValue());
-					}
-					else
-					{
-						Transcation.Cancel();
-					}
-				}
-			}
-
-			TUniquePtr<ICurveEditorKeyDragOperation> KeyDrag = CreateKeyDrag(CurveEditor->Selection.GetSelectionType());
-
-			bool bAllowSnapping = NewPoint.IsSet();
+		MenuBuilder.EndSection();
+	}
 
-			KeyDrag->Initialize(CurveEditor.Get(), NewPoint);
-			KeyDrag->SnapMetrics.bSnapInputValues = bAllowSnapping;
-			KeyDrag->SnapMetrics.bSnapOutputValues = bAllowSnapping;
 
-			DragOperation = FCurveEditorDelayedDrag(MousePixel, MouseEvent.GetEffectingButton());
-			DragOperation->DragImpl = MoveTemp(KeyDrag);
+	return MenuBuilder.MakeWidget();
+}
 
-			return FReply::Handled().CaptureMouse(AsShared());
-		}
-	}
-	else if (MouseEvent.GetEffectingButton() == EKeys::RightMouseButton)
+FSlateIcon SCurveEditorPanel::GetCurveExtrapolationPreIcon() const
+{
+	// We check to see if pre/post share a extrapolation mode and return a shared icon, otherwise mixed.
+	if (CompareCommonPreExtrapolationMode(RCCE_Constant))
 	{
-		if (MouseEvent.IsAltDown())
-		{
-			DragOperation = FCurveEditorDelayedDrag(MousePixel, MouseEvent.GetEffectingButton());
-			DragOperation->DragImpl = MakeUnique<FCurveEditorDragOperation_Zoom>(CurveEditor.Get());
-			return FReply::Handled().CaptureMouse(AsShared());
-		}
+		return FCurveEditorCommands::Get().SetPreInfinityExtrapConstant->GetIcon();
 	}
-
-	bool bShiftPressed = MouseEvent.IsShiftDown();
-	bool bCtrlPressed  = MouseEvent.IsControlDown();
-
-	TOptional<FCurvePointHandle> MouseDownPoint = HitPoint(MousePixel);
-	if (MouseDownPoint.IsSet())
+	else if (CompareCommonPreExtrapolationMode(RCCE_Cycle))
 	{
-		if (bShiftPressed)
-		{
-			CurveEditor->Selection.Add(MouseDownPoint.GetValue());
-		}
-		else if (bCtrlPressed)
-		{
-			CurveEditor->Selection.Toggle(MouseDownPoint.GetValue());
-		}
-		else
-		{
-			if (CurveEditor->Selection.Contains(MouseDownPoint->CurveID, MouseDownPoint->KeyHandle))
-			{
-				CurveEditor->Selection.ChangeSelectionPointType(MouseDownPoint->PointType);
-			}
-			else
-			{
-				CurveEditor->Selection.Clear();
-				CurveEditor->Selection.Add(MouseDownPoint.GetValue());
-			}
-		}
-
-		if (MouseEvent.GetEffectingButton() == EKeys::LeftMouseButton)
-		{
-			TUniquePtr<ICurveEditorKeyDragOperation> KeyDrag = CreateKeyDrag(MouseDownPoint->PointType);
-
-			KeyDrag->Initialize(CurveEditor.Get(), MouseDownPoint);
-
-			DragOperation = FCurveEditorDelayedDrag(MousePixel, MouseEvent.GetEffectingButton());
-			DragOperation->DragImpl = MoveTemp(KeyDrag);
-
-			return FReply::Handled().CaptureMouse(AsShared());
-		}
+		return FCurveEditorCommands::Get().SetPreInfinityExtrapCycle->GetIcon();
 	}
-
-	if (MouseEvent.GetEffectingButton() == EKeys::LeftMouseButton)
+	else if (CompareCommonPreExtrapolationMode(RCCE_CycleWithOffset))
 	{
-		DragOperation = FCurveEditorDelayedDrag(MousePixel, MouseEvent.GetEffectingButton());
-		DragOperation->DragImpl = MakeUnique<FCurveEditorDragOperation_Marquee>(CurveEditor.Get(), this);
-		return FReply::Handled().CaptureMouse(AsShared());
+		return FCurveEditorCommands::Get().SetPreInfinityExtrapCycleWithOffset->GetIcon();
 	}
-	else if (MouseEvent.GetEffectingButton() == EKeys::RightMouseButton)
+	else if (CompareCommonPreExtrapolationMode(RCCE_Linear))
 	{
-		DragOperation = FCurveEditorDelayedDrag(MousePixel, MouseEvent.GetEffectingButton());
-		DragOperation->DragImpl = MakeUnique<FCurveEditorDragOperation_Pan>(CurveEditor.Get());
-
-		return FReply::Handled().CaptureMouse(AsShared());
+		return FCurveEditorCommands::Get().SetPreInfinityExtrapLinear->GetIcon();
 	}
-
-	return FReply::Handled();
-}
-
-FReply SCurveEditorPanel::OnMouseMove(const FGeometry& MyGeometry, const FPointerEvent& MouseEvent)
-{
-	FVector2D MousePixel = MyGeometry.AbsoluteToLocal(MouseEvent.GetScreenSpacePosition());
-
-	if (DragOperation.IsSet())
+	else if (CompareCommonPreExtrapolationMode(RCCE_Oscillate))
 	{
-		FVector2D InitialPosition = DragOperation->GetInitialPosition();
-
-		if (!DragOperation->IsDragging() && DragOperation->AttemptDragStart(MouseEvent))
-		{
-			DragOperation->DragImpl->BeginDrag(InitialPosition, MousePixel, MouseEvent);
-		}
-		else if (DragOperation->IsDragging())
-		{
-			DragOperation->DragImpl->Drag(InitialPosition, MousePixel, MouseEvent);
-		}
+		return FCurveEditorCommands::Get().SetPreInfinityExtrapOscillate->GetIcon();
 	}
 	else
 	{
-		UpdateCurveProximities(MousePixel);
+		return FSlateIcon(FEditorStyle::GetStyleSetName(), "GenericCurveEditor.PreInfinityMixed");
 	}
-
-	return FReply::Handled();
 }
 
-FReply SCurveEditorPanel::OnMouseButtonUp(const FGeometry& MyGeometry, const FPointerEvent& MouseEvent)
+FSlateIcon SCurveEditorPanel::GetCurveExtrapolationPostIcon() const
 {
-	FVector2D MousePosition = MyGeometry.AbsoluteToLocal(MouseEvent.GetScreenSpacePosition());
-
-	if (DragOperation.IsSet() && DragOperation->IsDragging())
+	// We check to see if pre/post share a extrapolation mode and return a shared icon, otherwise mixed.
+	if (CompareCommonPostExtrapolationMode(RCCE_Constant))
+	{
+		return FCurveEditorCommands::Get().SetPostInfinityExtrapConstant->GetIcon();
+	}
+	else if (CompareCommonPostExtrapolationMode(RCCE_Cycle))
+	{
+		return FCurveEditorCommands::Get().SetPostInfinityExtrapCycle->GetIcon();
+	}
+	else if (CompareCommonPostExtrapolationMode(RCCE_CycleWithOffset))
+	{
+		return FCurveEditorCommands::Get().SetPostInfinityExtrapCycleWithOffset->GetIcon();
+	}
+	else if (CompareCommonPostExtrapolationMode(RCCE_Linear))
+	{
+		return FCurveEditorCommands::Get().SetPostInfinityExtrapLinear->GetIcon();
+	}
+	else if (CompareCommonPostExtrapolationMode(RCCE_Oscillate))
 	{
-		FVector2D InitialPosition = DragOperation->GetInitialPosition();
-		DragOperation->DragImpl->EndDrag(InitialPosition, MousePosition, MouseEvent);
+		return FCurveEditorCommands::Get().SetPostInfinityExtrapOscillate->GetIcon();
 	}
 	else
 	{
-		TOptional<FCurvePointHandle> MouseUpPoint = HitPoint(MousePosition);
-
-		RebindContextualActions(MousePosition);
-
-		if (MouseEvent.GetEffectingButton() == EKeys::RightMouseButton)
-		{
-			if (!MouseUpPoint.IsSet())
-			{
-				CurveEditor->Selection.Clear();
-			}
-
-			const bool bCloseAfterSelection = true;
-			FMenuBuilder MenuBuilder(bCloseAfterSelection, CommandList);
-
-			FCurveEditorContextMenu::BuildMenu(MenuBuilder, CurveEditor, MouseUpPoint, GetHoveredCurve());
-
-			// Push the context menu
-			FWidgetPath WidgetPath = MouseEvent.GetEventPath() != nullptr ? *MouseEvent.GetEventPath() : FWidgetPath();
-			FSlateApplication::Get().PushMenu(AsShared(), WidgetPath, MenuBuilder.MakeWidget(), FSlateApplication::Get().GetCursorPos(), FPopupTransitionEffect(FPopupTransitionEffect::ContextMenu));
-		}
-		else if (MouseEvent.GetEffectingButton() == EKeys::LeftMouseButton && !MouseUpPoint.IsSet())
-		{
-			if (!CommandList->ProcessCommandBindings(MouseEvent))
-			{
-				CurveEditor->Selection.Clear();
-			}
-		}
+		return FSlateIcon(FEditorStyle::GetStyleSetName(), "GenericCurveEditor.PostInfinityMixed");
 	}
-
-	UpdateCurveProximities(MousePosition);
-
-	DragOperation.Reset();
-	return FReply::Handled().ReleaseMouseCapture();
 }
 
-FReply SCurveEditorPanel::OnMouseWheel(const FGeometry& MyGeometry, const FPointerEvent& MouseEvent)
-{
-	FCurveEditorScreenSpace ScreenSpace = CurveEditor->GetScreenSpace();
-
-	FVector2D MousePixel   = MyGeometry.AbsoluteToLocal(MouseEvent.GetScreenSpacePosition());
-	double    CurrentTime  = ScreenSpace.ScreenToSeconds(MousePixel.X);
-	double    CurrentValue = ScreenSpace.ScreenToValue(MousePixel.Y);
-
-	float ZoomDelta = 1.f - FMath::Clamp(0.1f * MouseEvent.GetWheelDelta(), -0.9f, 0.9f);
-	CurveEditor->ZoomAround(ZoomDelta, CurrentTime, CurrentValue);
 
-	return FReply::Handled();
-}
-
-void SCurveEditorPanel::OnMouseEnter(const FGeometry& MyGeometry, const FPointerEvent& MouseEvent)
+void SCurveEditorPanel::ShowCurveFilterUI(TSubclassOf<UCurveEditorFilterBase> FilterClass)
 {
-	CurveProximities.Reset();
-
-	SCompoundWidget::OnMouseEnter(MyGeometry, MouseEvent);
-}
-
-void SCurveEditorPanel::OnMouseLeave(const FPointerEvent& MouseEvent)
-{
-	CurveProximities.Reset();
-
-	SCompoundWidget::OnMouseLeave(MouseEvent);
-}
-
-void SCurveEditorPanel::OnFocusLost(const FFocusEvent& InFocusEvent)
-{
-	if (DragOperation.IsSet())
+	TSharedPtr<FTabManager> TabManager = WeakTabManager.Pin();
+	if (TabManager)
 	{
-		DragOperation->DragImpl->CancelDrag();
-		DragOperation.Reset();
+		SCurveEditorFilterPanel::OpenDialog(TabManager.ToSharedRef(), CurveEditor.ToSharedRef(), FilterClass);
 	}
-
-	SCompoundWidget::OnFocusLost(InFocusEvent);
 }
 
-FReply SCurveEditorPanel::OnKeyDown(const FGeometry& MyGeometry, const FKeyEvent& InKeyEvent)
+const FGeometry& SCurveEditorPanel::GetScrollPanelGeometry() const
 {
-	if (InKeyEvent.GetKey() == EKeys::Escape)
-	{
-		if (DragOperation.IsSet())
-		{
-			DragOperation->DragImpl->CancelDrag();
-			DragOperation.Reset();
-		}
-		else
-		{
-			CurveEditor->Selection.Clear();
-		}
-		return FReply::Handled();
-	}
-	else if (CommandList->ProcessCommandBindings(InKeyEvent))
-	{
-		return FReply::Handled();
-	}
-	return FReply::Unhandled();
+	return ScrollBox->GetCachedGeometry();
 }
 
-int32 SCurveEditorPanel::OnPaint(const FPaintArgs& Args, const FGeometry& AllottedGeometry, const FSlateRect& MyCullingRect, FSlateWindowElementList& OutDrawElements, int32 LayerId, const FWidgetStyle& InWidgetStyle, bool bParentEnabled) const
+const FGeometry& SCurveEditorPanel::GetViewContainerGeometry() const
 {
-	const ESlateDrawEffect DrawEffects = ShouldBeEnabled(bParentEnabled) ? ESlateDrawEffect::None : ESlateDrawEffect::DisabledEffect;
-
-	FSlateDrawElement::MakeBox(OutDrawElements,	LayerId, AllottedGeometry.ToPaintGeometry(),
-		FEditorStyle::GetBrush("ToolPanel.GroupBorder"), DrawEffects);
-
-	LayerId = DrawGridLines(AllottedGeometry, OutDrawElements, LayerId, DrawEffects);
-	LayerId = DrawCurves(AllottedGeometry, OutDrawElements, LayerId, InWidgetStyle, DrawEffects);
-
-	if (DragOperation.IsSet() && DragOperation->IsDragging())
-	{
-		LayerId = DragOperation->DragImpl->Paint(AllottedGeometry, OutDrawElements, LayerId);
-	}
-	return SCompoundWidget::OnPaint(Args, AllottedGeometry, MyCullingRect, OutDrawElements, LayerId+2000, InWidgetStyle, bParentEnabled)-2000;
+	return CurveViewsContainer->GetCachedGeometry();
 }
 
-int32 SCurveEditorPanel::DrawGridLines(const FGeometry& AllottedGeometry, FSlateWindowElementList& OutDrawElements, int32 LayerId, ESlateDrawEffect DrawEffects) const
+TSharedPtr<FExtender> SCurveEditorPanel::GetToolbarExtender()
 {
-	// Rendering info
-	const float          Width          = AllottedGeometry.GetLocalSize().X;
-	const float          Height         = AllottedGeometry.GetLocalSize().Y;
-	const FLinearColor   MajorGridColor = GridLineTintAttribute.Get();
-	const FLinearColor   MinorGridColor = MajorGridColor.CopyWithNewOpacity(MajorGridColor.A * .25f);
-	const FPaintGeometry PaintGeometry  = AllottedGeometry.ToPaintGeometry();
-	const FLinearColor	 LabelColor = FLinearColor::White.CopyWithNewOpacity(0.65f);
-	const FSlateFontInfo FontInfo = FCoreStyle::Get().GetFontStyle("ToolTip.LargerFont");
-
-	TArray<float> MajorGridLines, MinorGridLines;
-	TArray<FText> MajorGridLabels;
-	CurveEditor->GetGridLinesX(MajorGridLines, MinorGridLines, MajorGridLabels);
-	ensureMsgf(MajorGridLines.Num() == MajorGridLabels.Num(), TEXT("A grid label should be specified for every major grid line, even if it is just an empty FText."));
-
-
-	TArray<FVector2D> LinePoints;
-	LinePoints.Add(FVector2D(0.f, 0.f));
-	LinePoints.Add(FVector2D(0.f, Height));
-
-	// Draw major grid lines
-	for(int32 i = 0; i < MajorGridLines.Num(); i++)
+	// We're going to create a new Extender and add the main Curve Editor icons to it.
+	// We combine this with the extender provided by the Curve Editor Module as that extender has been extended by tools
+	ICurveEditorModule& CurveEditorModule = FModuleManager::Get().LoadModuleChecked<ICurveEditorModule>("CurveEditor");
+	TArray<TSharedPtr<FExtender>> ToolbarExtenders;
+	for (ICurveEditorModule::FCurveEditorMenuExtender& ExtenderCallback : CurveEditorModule.GetAllToolBarMenuExtenders())
 	{
-		LinePoints[0].X = LinePoints[1].X = MajorGridLines[i];
-
-		FSlateDrawElement::MakeLines(
-			OutDrawElements,
-			LayerId,
-			PaintGeometry,
-			LinePoints,
-			DrawEffects,
-			MajorGridColor,
-			false
-		);
-
-		FPaintGeometry LabelGeometry = AllottedGeometry.ToPaintGeometry( FSlateLayoutTransform(FVector2D(LinePoints[0].X + LabelOffsetPx, LabelOffsetPx)) );
-
-		// Draw the axis labels above the curves (which get drawn next) now so that we don't
-		// have to cache off our labels and draw after. Done as +2 as the curves will draw on
-		// +1 after this function returns.
-		FSlateDrawElement::MakeText(
-			OutDrawElements,
-			LayerId+2,
-			LabelGeometry,
-			MajorGridLabels[i],
-			FontInfo,
-			DrawEffects,
-			LabelColor
-		);
+		ToolbarExtenders.Add(ExtenderCallback.Execute(GetCommands().ToSharedRef()));
 	}
+	TSharedPtr<FExtender> Extender = FExtender::Combine(ToolbarExtenders);
 
-	for (float PosX : MinorGridLines)
+	struct Local
 	{
-		LinePoints[0].X = LinePoints[1].X = PosX;
-
-		FSlateDrawElement::MakeLines(
-			OutDrawElements,
-			LayerId,
-			PaintGeometry,
-			LinePoints,
-			DrawEffects,
-			MinorGridColor,
-			false
-		);
-	}
-
-	MajorGridLines.Reset();
-	MinorGridLines.Reset();
-	MajorGridLabels.Reset();
-	CurveEditor->GetGridLinesY(MajorGridLines, MinorGridLines, MajorGridLabels);
-	ensureMsgf(MajorGridLines.Num() == MajorGridLabels.Num(), TEXT("A grid label should be specified for every major grid line, even if it is just an empty FText."));
-
-	LinePoints[0].X = 0.f;
-	LinePoints[1].X = Width;
-	for (int32 i = 0; i < MajorGridLines.Num(); i++)
-	{
-		LinePoints[0].Y = LinePoints[1].Y = MajorGridLines[i];
-
-		FSlateDrawElement::MakeLines(
-			OutDrawElements,
-			LayerId,
-			PaintGeometry,
-			LinePoints,
-			DrawEffects,
-			MajorGridColor,
-			false
-		);
-
-		FPaintGeometry LabelGeometry = AllottedGeometry.ToPaintGeometry(FSlateLayoutTransform(FVector2D(LabelOffsetPx, LinePoints[0].Y+ LabelOffsetPx)));
-
-		// Draw the axis labels above the curves (which get drawn next) now so that we don't
-		// have to cache off our labels and draw after. Done as +2 as the curves will draw on
-		// +1 after this function returns.
-		FSlateDrawElement::MakeText(
-			OutDrawElements,
-			LayerId+2,
-			LabelGeometry,
-			MajorGridLabels[i],
-			FontInfo,
-			DrawEffects,
-			LabelColor
-		);
-	}
+		static void FillToolbar(FToolBarBuilder& ToolBarBuilder, TSharedRef<SCurveKeyDetailPanel> InKeyDetailsPanel, TSharedRef<SCurveEditorPanel> InEditorPanel)
+		{
 
-	for (float PosX : MinorGridLines)
-	{
-		LinePoints[0].X = LinePoints[1].X = PosX;
-
-		FSlateDrawElement::MakeLines(
-			OutDrawElements,
-			LayerId,
-			PaintGeometry,
-			LinePoints,
-			DrawEffects,
-			MinorGridColor,
-			false
-		);
-	}
+			ToolBarBuilder.BeginSection("Tools");
+			{
+				ToolBarBuilder.AddToolBarButton(FCurveEditorCommands::Get().DeactivateCurrentTool);
+			}
+			ToolBarBuilder.EndSection();
 
-	return LayerId + 1;
-}
+			ToolBarBuilder.BeginSection("Key Details");
+			{
+				ToolBarBuilder.AddWidget(InKeyDetailsPanel);
+			}
+			ToolBarBuilder.EndSection();
 
-int32 SCurveEditorPanel::DrawCurves(const FGeometry& AllottedGeometry, FSlateWindowElementList& OutDrawElements, int32 LayerId, const FWidgetStyle& InWidgetStyle, ESlateDrawEffect DrawEffects) const
-{
-	FCurveEditorScreenSpace ScreenSpace = CurveEditor->GetScreenSpace();
+			ToolBarBuilder.BeginSection("View");
+			{
+				// Dropdown Menu for choosing your viewing mode
+				TAttribute<FSlateIcon> ViewModeIcon;
+				ViewModeIcon.Bind(TAttribute<FSlateIcon>::FGetter::CreateLambda([InEditorPanel] {
+					switch (InEditorPanel->GetViewMode())
+					{
+					case ECurveEditorViewID::Absolute:
+						return FCurveEditorCommands::Get().SetViewModeAbsolute->GetIcon();
+					case ECurveEditorViewID::Stacked:
+						return FCurveEditorCommands::Get().SetViewModeStacked->GetIcon();
+					case ECurveEditorViewID::Normalized:
+						return FCurveEditorCommands::Get().SetViewModeNormalized->GetIcon();
+					default: // EKeyGroupMode::None
+						return FCurveEditorCommands::Get().SetAxisSnappingNone->GetIcon();
+					}
+				}));
+
+				ToolBarBuilder.AddComboButton(
+					FUIAction(),
+					FOnGetContent::CreateSP(InEditorPanel, &SCurveEditorPanel::MakeViewModeMenu),
+					LOCTEXT("ViewModeDropdown", "Curve View Modes"),
+					LOCTEXT("ViewModeDropdownToolTip", "Choose the viewing mode for the curves."),
+					ViewModeIcon);
+			}
+			ToolBarBuilder.EndSection();
 
-	static const FName SelectionColorName("SelectionColor");
-	FLinearColor SelectionColor = FEditorStyle::GetSlateColor(SelectionColorName).GetColor(InWidgetStyle);
+			ToolBarBuilder.BeginSection("Framing");
+			{
+				// Framing
+				ToolBarBuilder.AddToolBarButton(FCurveEditorCommands::Get().ZoomToFit);
+			}
+			ToolBarBuilder.EndSection();
 
-	const FVector2D      VisibleSize   = AllottedGeometry.GetLocalSize();
-	const FPaintGeometry PaintGeometry = AllottedGeometry.ToPaintGeometry();
+			ToolBarBuilder.BeginSection("Visibility");
+			{
+				// Curve Visibility
+				ToolBarBuilder.AddComboButton(
+					FUIAction(),
+					FOnGetContent::CreateSP(InEditorPanel, &SCurveEditorPanel::MakeCurveEditorCurveViewOptionsMenu),
+					LOCTEXT("CurveEditorCurveOptions", "Curves Options"),
+					LOCTEXT("CurveEditorCurveOptionsToolTip", "Curve Options"),
+					FSlateIcon(FEditorStyle::GetStyleSetName(), "GenericCurveEditor.VisibilityOptions"));
 
-	TOptional<FCurveModelID> HoveredCurve = GetHoveredCurve();
-	for (const FCurveDrawParams& Params : CachedDrawParams)
-	{
-		float Thickness = ( HoveredCurve.IsSet() && HoveredCurve.GetValue() == Params.GetID() ) ? HoveredCurveThickness : UnHoveredCurveThickness;
-		FSlateDrawElement::MakeLines(
-			OutDrawElements,
-			LayerId,
-			PaintGeometry,
-			Params.InterpolatingPoints,
-			DrawEffects,
-			Params.Color,
-			bAntiAliasCurves,
-			Thickness
-		);
-
-		TArray<FVector2D> LinePoints;
-		LinePoints.SetNum(2);
-
-		// Draw tangents
-		for (const FCurvePointInfo& Point : Params.Points)
-		{
-			const FKeyDrawInfo& PointDrawInfo = Params.GetKeyDrawInfo(Point.Type);
-			const bool          bSelected     = CurveEditor->Selection.IsSelected(FCurvePointHandle(Params.GetID(), Point.Type, Point.KeyHandle));
-			const FLinearColor  PointTint     = bSelected ? SelectionColor : PointDrawInfo.Tint;
+			}
+			ToolBarBuilder.EndSection();
 
-			if (Point.LineDelta.X != 0.f || Point.LineDelta.Y != 0.f)
+			ToolBarBuilder.BeginSection("Adjustment");
 			{
-				LinePoints[0] = Point.ScreenPosition;
-				LinePoints[1] = Point.ScreenPosition + Point.LineDelta;
-
-				// Draw the connecting line
-				FSlateDrawElement::MakeLines(OutDrawElements, LayerId, PaintGeometry, LinePoints, DrawEffects, PointTint, true);
+				// Toggle Button for Time Snapping
+				ToolBarBuilder.AddToolBarButton(FCurveEditorCommands::Get().ToggleInputSnapping);
+
+				// Dropdown Menu to choose the snapping scale.
+				FUIAction TimeSnapMenuAction(FExecuteAction(), FCanExecuteAction::CreateLambda([InEditorPanel] { return !InEditorPanel->DisabledTimeSnapTooltipAttribute.IsSet(); }));
+				ToolBarBuilder.AddComboButton(
+					TimeSnapMenuAction,
+					FOnGetContent::CreateSP(InEditorPanel, &SCurveEditorPanel::MakeTimeSnapMenu),
+					LOCTEXT("TimeSnappingOptions", "Time Snapping"),
+					TAttribute<FText>(InEditorPanel, &SCurveEditorPanel::GetTimeSnapMenuTooltip),
+					TAttribute<FSlateIcon>(),
+					true);
+
+				// Toggle Button for Value Snapping
+				ToolBarBuilder.AddToolBarButton(FCurveEditorCommands::Get().ToggleOutputSnapping);
+
+				// Dropdown Menu to choose the snapping scale.
+				ToolBarBuilder.AddComboButton(
+					FUIAction(),
+					FOnGetContent::CreateSP(InEditorPanel, &SCurveEditorPanel::MakeValueSnapMenu),
+					LOCTEXT("ValueSnappingOptions", "Value Snapping"),
+					LOCTEXT("ValueSnappingOptionsToolTip", "Choose what precision the Value axis is snapped to while moving keys."),
+					TAttribute<FSlateIcon>(),
+					true);
+
+
+				// Dropdown Menu for choosing your axis snapping for tool movement
+				TAttribute<FSlateIcon> AxisSnappingModeIcon;
+				AxisSnappingModeIcon.Bind(TAttribute<FSlateIcon>::FGetter::CreateLambda([InEditorPanel] {
+					switch (InEditorPanel->GetCurveEditor()->GetAxisSnap().RestrictedAxisList)
+					{
+					case EAxisList::Type::X:
+						return FCurveEditorCommands::Get().SetAxisSnappingHorizontal->GetIcon();
+					case EAxisList::Type::Y:
+						return FCurveEditorCommands::Get().SetAxisSnappingVertical->GetIcon();
+					default: // EKeyGroupMode::None
+						return FCurveEditorCommands::Get().SetAxisSnappingNone->GetIcon();
+					}
+				}));
+
+				ToolBarBuilder.AddComboButton(
+					FUIAction(),
+					FOnGetContent::CreateSP(InEditorPanel, &SCurveEditorPanel::MakeAxisSnapMenu),
+					LOCTEXT("AxisSnappingOptions", "Axis Snapping"),
+					LOCTEXT("AxisSnappingOptionsToolTip", "Choose which axes movement tools are locked to."),
+					AxisSnappingModeIcon);
 			}
+			ToolBarBuilder.EndSection();
 
-			FPaintGeometry PointGeometry = AllottedGeometry.ToPaintGeometry(
-				Point.ScreenPosition - (PointDrawInfo.ScreenSize * 0.5f),
-				PointDrawInfo.ScreenSize
-			);
+			ToolBarBuilder.BeginSection("Tangents");
+			{
+				ToolBarBuilder.AddToolBarButton(FCurveEditorCommands::Get().InterpolationCubicAuto);
+				ToolBarBuilder.AddToolBarButton(FCurveEditorCommands::Get().InterpolationCubicUser);
+				ToolBarBuilder.AddToolBarButton(FCurveEditorCommands::Get().InterpolationCubicBreak);
+				ToolBarBuilder.AddToolBarButton(FCurveEditorCommands::Get().InterpolationLinear);
+				ToolBarBuilder.AddToolBarButton(FCurveEditorCommands::Get().InterpolationConstant);
+				ToolBarBuilder.AddToolBarButton(FCurveEditorCommands::Get().InterpolationToggleWeighted);
+				
+				// We re-use key interpolation checks here, as you can set them under the same conditions.
+				FCanExecuteAction CanSetInfinities = FCanExecuteAction::CreateSP(InEditorPanel, &SCurveEditorPanel::CanSetKeyInterpolation);
+				
+				// Dropdown Menu for choosing your pre infinity for selected curves.
+				ToolBarBuilder.AddComboButton(
+					FUIAction(FExecuteAction(), CanSetInfinities),
+					FOnGetContent::CreateSP(InEditorPanel, &SCurveEditorPanel::MakeCurveExtrapolationMenu, false),
+					LOCTEXT("CurveEditorPreInfinityOptions", "Pre Infinity"),
+					LOCTEXT("CurveEditorPreInfinityOptionsToolTip", "Choose how the curve is evaluated if sampled before the first key"),
+					TAttribute<FSlateIcon>(InEditorPanel, &SCurveEditorPanel::GetCurveExtrapolationPreIcon)
+				);
+
+					// Dropdown menu for choosing your post infinity for selected curves
+				ToolBarBuilder.AddComboButton(
+					FUIAction(FExecuteAction(), CanSetInfinities),
+					FOnGetContent::CreateSP(InEditorPanel, &SCurveEditorPanel::MakeCurveExtrapolationMenu, true),
+					LOCTEXT("CurveEditorPostInfinityOptions", "Post Infinity"),
+					LOCTEXT("CurveEditorPostInfinityOptionsToolTip", "Choose how the curve is evaluated if sampled after the last key"),
+					TAttribute<FSlateIcon>(InEditorPanel, &SCurveEditorPanel::GetCurveExtrapolationPostIcon)
+				);
+			}
+			ToolBarBuilder.EndSection();
 
-			FSlateDrawElement::MakeBox(OutDrawElements, LayerId+Point.LayerBias, PointGeometry, PointDrawInfo.Brush, DrawEffects, PointTint);
+			ToolBarBuilder.BeginSection("Filters");
+			{
+				ToolBarBuilder.AddToolBarButton(FCurveEditorCommands::Get().FlattenTangents);
+				ToolBarBuilder.AddToolBarButton(FCurveEditorCommands::Get().StraightenTangents);
+				ToolBarBuilder.AddToolBarButton(FCurveEditorCommands::Get().OpenUserImplementableFilterWindow);
+			}
+			ToolBarBuilder.EndSection();
 		}
-	}
+	};
 
-	return LayerId + 1;
-}
+	Extender->AddToolBarExtension(
+		"Asset",
+		EExtensionHook::After,
+		GetCommands(),
+		FToolBarExtensionDelegate::CreateStatic(&Local::FillToolbar, KeyDetailsView.ToSharedRef(), SharedThis(this))
+	);
 
+	return Extender;
+}
 
-void SCurveEditorPanel::OnSimplifySelection()
+TSharedRef<SWidget> SCurveEditorPanel::MakeTimeSnapMenu()
 {
-	// Display dialog and let user enter tolerance.
-	GenericTextEntryModeless(
-		NSLOCTEXT("CurveEditor.Popups", "ReduceCurveTolerance", "Tolerance"),
-		FText::AsNumber(ReduceTolerance),
-		FOnTextCommitted::CreateSP(this, &SCurveEditorPanel::OnSimplifySelectionCommited)
-	);
+	TSharedRef<SWidget> InputSnapWidget =
+		SNew(SFrameRatePicker)
+		.Value_Lambda([this] { return this->CurveEditor->InputSnapRateAttribute.Get(); })
+		.OnValueChanged_Lambda([this](FFrameRate InFrameRate) { this->CurveEditor->InputSnapRateAttribute = InFrameRate; })
+		.PresetValues({
+		// We re-use the common frame rates but omit some of them.
+		FCommonFrameRateInfo{ FCommonFrameRates::FPS_12(),	LOCTEXT("Snap_Input_Twelve", "82ms (1/12s)"), LOCTEXT("Snap_Input_Description_Twelve", "Snap time values to one twelfth of a second (ie: 12fps)") },
+		FCommonFrameRateInfo{ FCommonFrameRates::FPS_15(),	LOCTEXT("Snap_Input_Fifteen", "66s (1/15s)"), LOCTEXT("Snap_Input_Description_Fifteen", "Snap time values to one fifteenth of a second (ie: 15fps)") },
+		FCommonFrameRateInfo{ FCommonFrameRates::FPS_24(),	LOCTEXT("Snap_Input_TwentyFour", "42ms (1/24s)"), LOCTEXT("Snap_Input_Description_TwentyFour", "Snap time values to one twenty-fourth of a second (ie: 24fps)") },
+		FCommonFrameRateInfo{ FCommonFrameRates::FPS_25(),	LOCTEXT("Snap_Input_TwentyFive", "40ms (1/25s)"), LOCTEXT("Snap_Input_Description_TwentyFive", "Snap time values to one twenty-fifth of a second (ie: 25fps)") },
+		FCommonFrameRateInfo{ FCommonFrameRates::FPS_30(),	LOCTEXT("Snap_Input_Thirty", "33ms (1/30s)"), LOCTEXT("Snap_Input_Description_Thirty", "Snap time values to one thirtieth of a second (ie: 30fps)") },
+		FCommonFrameRateInfo{ FCommonFrameRates::FPS_48(),	LOCTEXT("Snap_Input_FourtyEight", "21ms (1/48s)"), LOCTEXT("Snap_Input_Description_FourtyEight", "Snap time values to one fourth-eight of a second (ie: 48fps)") },
+		FCommonFrameRateInfo{ FCommonFrameRates::FPS_50(),	LOCTEXT("Snap_Input_Fifty", "20ms (1/50s)"), LOCTEXT("Snap_Input_Description_Fifty", "Snap time values to one fiftieth of a second (ie: 50fps)") },
+		FCommonFrameRateInfo{ FCommonFrameRates::FPS_60(),	LOCTEXT("Snap_Input_Sixty", "16ms (1/60s)"), LOCTEXT("Snap_Input_Description_Sixty", "Snap time values to one sixtieth of a second (ie: 60fps)") },
+		FCommonFrameRateInfo{ FCommonFrameRates::FPS_100(),	LOCTEXT("Snap_Input_OneHundred", "10ms (1/100s)"), LOCTEXT("Snap_Input_Description_OneHundred", "Snap time values to one one-hundredth of a second (ie: 100fps)") },
+		FCommonFrameRateInfo{ FCommonFrameRates::FPS_120(),	LOCTEXT("Snap_Input_OneHundredTwenty", "8ms (1/120s)"), LOCTEXT("Snap_Input_Description_OneHundredTwenty", "Snap time values to one one-hundred-twentieth of a second (ie: 120fps)") },
+		FCommonFrameRateInfo{ FCommonFrameRates::FPS_240(), LOCTEXT("Snap_Input_TwoHundredFourty", "4ms (1/240s)"), LOCTEXT("Snap_Input_Description_TwoHundredFourty", "Snap time values to one two-hundred-fourtieth of a second (ie: 240fps)") }
+			});
+
+	return InputSnapWidget;
 }
 
-void SCurveEditorPanel::OnSimplifySelectionCommited(const FText& InText, ETextCommit::Type CommitInfo)
+FText SCurveEditorPanel::GetTimeSnapMenuTooltip() const
 {
-	CloseEntryPopupMenu();
-	if (CommitInfo == ETextCommit::OnEnter)
+	// If this is specified then the time snap menu is disabled
+	if (DisabledTimeSnapTooltipAttribute.IsSet())
 	{
-		double NewTolerance = FCString::Atod(*InText.ToString());
-		const bool bIsNumber = InText.IsNumeric();
-		if (!bIsNumber)
-			return;
-		CurveEditor->SimplifySelection((float)NewTolerance);
+		return DisabledTimeSnapTooltipAttribute.Get();
 	}
-}
 
+	return LOCTEXT("TimeSnappingOptionsToolTip", "Choose what precision the Time axis is snapped to while moving keys.");
+}
 
+TSharedRef<SWidget> SCurveEditorPanel::MakeValueSnapMenu()
+{
+	TArray<SNumericList<float>::FNamedValue> ValueSnapAmounts;
+	// SnapValues.Add( SNumericDropDown<float>::FNamedValue( 0.001f, LOCTEXT( "Snap_OneThousandth", "0.001" ), LOCTEXT( "SnapDescription_OneThousandth", "Set snap to 1/1000th" ) ) );
+	//SnapValues.Add( SNumericDropDown<float>::FNamedValue( 0.01f, LOCTEXT( "Snap_OneHundredth", "0.01" ), LOCTEXT( "SnapDescription_OneHundredth", "Set snap to 1/100th" ) ) );
+	ValueSnapAmounts.Add(SNumericList<float>::FNamedValue(0.1f, LOCTEXT("Snap_OneTenth", "0.1"), LOCTEXT("SnapDescription_OneTenth", "Set snap to 1/10th")));
+	ValueSnapAmounts.Add(SNumericList<float>::FNamedValue(0.5f, LOCTEXT("Snap_OneHalf", "0.5"), LOCTEXT("SnapDescription_OneHalf", "Set snap to 1/2")));
+	ValueSnapAmounts.Add(SNumericList<float>::FNamedValue(1.0f, LOCTEXT("Snap_One", "1"), LOCTEXT("SnapDescription_One", "Set snap to 1")));
+	ValueSnapAmounts.Add(SNumericList<float>::FNamedValue(2.0f, LOCTEXT("Snap_Two", "2"), LOCTEXT("SnapDescription_Two", "Set snap to 2")));
+	ValueSnapAmounts.Add(SNumericList<float>::FNamedValue(5.0f, LOCTEXT("Snap_Five", "5"), LOCTEXT("SnapDescription_Five", "Set snap to 5")));
+	ValueSnapAmounts.Add(SNumericList<float>::FNamedValue(10.0f, LOCTEXT("Snap_Ten", "10"), LOCTEXT("SnapDescription_Ten", "Set snap to 10")));
+	ValueSnapAmounts.Add(SNumericList<float>::FNamedValue(50.0f, LOCTEXT("Snap_Fifty", "50"), LOCTEXT("SnapDescription_50", "Set snap to 50")));
+	ValueSnapAmounts.Add(SNumericList<float>::FNamedValue(100.0f, LOCTEXT("Snap_OneHundred", "100"), LOCTEXT("SnapDescription_OneHundred", "Set snap to 100")));
+
+	TSharedRef<SWidget> OutputSnapWidget =
+		SNew(SNumericList<float>)
+		.DropDownValues(ValueSnapAmounts)
+		.MinDesiredValueWidth(60)
+		.Value_Lambda([this]() -> float { return this->CurveEditor->OutputSnapIntervalAttribute.Get(); })
+		.OnValueChanged_Lambda([this](float InNewOutputSnap) { this->CurveEditor->OutputSnapIntervalAttribute.Set(InNewOutputSnap); })
+		.HeaderText(LOCTEXT("CurveEditorMenuValueSnapHeader", "Value Snap"));
+
+	return OutputSnapWidget;
+}
 
-void SCurveEditorPanel::GenericTextEntryModeless(const FText& DialogText, const FText& DefaultText, FOnTextCommitted OnTextComitted)
+TSharedRef<SWidget> SCurveEditorPanel::MakeAxisSnapMenu()
 {
-	TSharedRef<STextEntryPopup> TextEntryPopup =
-		SNew(STextEntryPopup)
-		.Label(DialogText)
-		.DefaultText(DefaultText)
-		.OnTextCommitted(OnTextComitted)
-		.ClearKeyboardFocusOnCommit(false)
-		.SelectAllTextWhenFocused(true)
-		.MaxWidth(1024.0f);
-
-	EntryPopupMenu = FSlateApplication::Get().PushMenu(
-		SharedThis(this),
-		FWidgetPath(),
-		TextEntryPopup,
-		FSlateApplication::Get().GetCursorPos(),
-		FPopupTransitionEffect(FPopupTransitionEffect::TypeInPopup)
-	);
+	// This builds the dropdown menu when looking at the Curve View Options combobox.
+	FMenuBuilder MenuBuilder(true, GetCommands());
+
+	MenuBuilder.AddMenuEntry(FCurveEditorCommands::Get().SetAxisSnappingNone);
+	MenuBuilder.AddMenuEntry(FCurveEditorCommands::Get().SetAxisSnappingHorizontal);
+	MenuBuilder.AddMenuEntry(FCurveEditorCommands::Get().SetAxisSnappingVertical);
+	
+	return MenuBuilder.MakeWidget();
 }
 
+TSharedRef<SWidget> SCurveEditorPanel::MakeViewModeMenu()
+{
+	// This builds the dropdown menu when looking at the Curve View Modes
+	FMenuBuilder MenuBuilder(true, GetCommands());
+
+	MenuBuilder.AddMenuEntry(FCurveEditorCommands::Get().SetViewModeAbsolute);
+	MenuBuilder.AddMenuEntry(FCurveEditorCommands::Get().SetViewModeStacked);
+	MenuBuilder.AddMenuEntry(FCurveEditorCommands::Get().SetViewModeNormalized);
+
+	return MenuBuilder.MakeWidget();
+}
 
-void SCurveEditorPanel::CloseEntryPopupMenu()
+bool SCurveEditorPanel::IsInlineEditPanelEditable() const
 {
-	if (EntryPopupMenu.IsValid())
-	{
-		EntryPopupMenu.Pin()->Dismiss();
-	}
+	return CurveEditor->GetSelection().Count() > 0;
 }
 
+EVisibility SCurveEditorPanel::ShouldInstructionOverlayBeVisible() const
+{
+	// The instruction overlay is visible if they have no selection in the tree.
+	const bool bCurvesAreVisible = CurveEditor->GetTreeSelection().Num() > 0 || CurveEditor->GetPinnedCurves().Num() > 0;
+	return bCurvesAreVisible ? EVisibility::Hidden : EVisibility::HitTestInvisible;
+}
 #undef LOCTEXT_NAMESPACE
\ No newline at end of file
diff --git a/Engine/Source/Editor/CurveEditor/Private/SCurveEditorViewContainer.h b/Engine/Source/Editor/CurveEditor/Private/SCurveEditorViewContainer.h
new file mode 100644
--- /dev/null
+++ b/Engine/Source/Editor/CurveEditor/Private/SCurveEditorViewContainer.h
@@ -0,0 +1,82 @@
+// Copyright 1998-2019 Epic Games, Inc. All Rights Reserved.
+
+#pragma once
+
+#include "Widgets/DeclarativeSyntaxSupport.h"
+#include "Widgets/SCompoundWidget.h"
+
+#include "ICurveEditorDragOperation.h"
+#include "CurveEditor.h"
+#include "CurveEditorTypes.h"
+
+class SScrollBox;
+class SCurveEditorView;
+class ITimeSliderController;
+
+/**
+ * Curve editor widget that reflects the state of an FCurveEditor
+ */
+class CURVEEDITOR_API SCurveEditorViewContainer : public SVerticalBox
+{
+	SLATE_BEGIN_ARGS(SCurveEditorViewContainer)
+	{}
+
+		/** Optional Time Slider Controller which allows us to synchronize with an externally controlled Time Slider */
+		SLATE_ARGUMENT(TSharedPtr<ITimeSliderController>, ExternalTimeSliderController)
+
+	SLATE_END_ARGS()
+
+	/**
+	 * Construct a new curve editor panel widget
+	 */
+	void Construct(const FArguments& InArgs, TSharedRef<FCurveEditor> InCurveEditor);
+
+	TArrayView<const TSharedPtr<SCurveEditorView>> GetViews() const { return Views; }
+
+	void AddView(TSharedRef<SCurveEditorView> ViewToAdd);
+
+	void Clear();
+
+private:
+
+	// SWidget Interface
+	virtual int32 OnPaint(const FPaintArgs& Args, const FGeometry& AllottedGeometry, const FSlateRect& MyCullingRect, FSlateWindowElementList& OutDrawElements, int32 LayerId, const FWidgetStyle& InWidgetStyle, bool bParentEnabled) const override;
+	virtual void Tick(const FGeometry& AllottedGeometry, const double InCurrentTime, const float InDeltaTime) override;
+	virtual FVector2D ComputeDesiredSize(float) const;
+	virtual bool ComputeVolatility() const override;
+	// ~SWidget Interface
+
+	/*~ Mouse interaction */
+	virtual FReply OnPreviewMouseButtonDown(const FGeometry& MyGeometry, const FPointerEvent& MouseEvent) override;
+	virtual FReply OnMouseButtonDown(const FGeometry& MyGeometry, const FPointerEvent& MouseEvent) override;
+	virtual FReply OnMouseButtonUp(const FGeometry& MyGeometry, const FPointerEvent& MouseEvent) override;
+	virtual FReply OnMouseMove(const FGeometry& MyGeometry, const FPointerEvent& MouseEvent) override;
+	virtual FReply OnMouseButtonDoubleClick(const FGeometry& InMyGeometry, const FPointerEvent& InMouseEvent) override;
+	virtual void OnFocusLost(const FFocusEvent& InFocusEvent) override;
+	virtual FCursorReply OnCursorQuery(const FGeometry& MyGeometry, const FPointerEvent& CursorEvent) const override;
+
+	virtual bool SupportsKeyboardFocus() const override { return true; }
+	virtual FReply OnKeyDown(const FGeometry& MyGeometry, const FKeyEvent& InKeyEvent) override;
+
+private:
+
+	FMargin GetSlotPadding(int32 SlotIndex) const;
+
+	void OnCurveEditorToolChanged(FCurveEditorToolID InToolId);
+
+	void ExpandInputBounds(float NewWidth);
+
+private:
+
+	/** The curve editor pointer */
+	TSharedPtr<FCurveEditor> CurveEditor;
+
+	/** Optional time slider controller */
+	TSharedPtr<ITimeSliderController> TimeSliderController;
+
+	/** (Optional) the current drag operation */
+	TOptional<FCurveEditorDelayedDrag> DragOperation;
+
+	/** Array of views that may need their height updating on tick. */
+	TArray<TSharedPtr<SCurveEditorView>> Views;
+};
\ No newline at end of file
diff --git a/Engine/Source/Editor/CurveEditor/Private/Tree/SCurveEditorTree.cpp b/Engine/Source/Editor/CurveEditor/Private/Tree/SCurveEditorTree.cpp
new file mode 100644
--- /dev/null
+++ b/Engine/Source/Editor/CurveEditor/Private/Tree/SCurveEditorTree.cpp
@@ -0,0 +1,194 @@
+// Copyright 1998-2019 Epic Games, Inc. All Rights Reserved.
+
+#include "Tree/SCurveEditorTree.h"
+#include "Tree/CurveEditorTree.h"
+#include "Tree/ICurveEditorTreeItem.h"
+
+#include "CurveEditor.h"
+
+#include "Widgets/Views/STableRow.h"
+#include "Widgets/Views/STreeView.h"
+#include "Widgets/Views/SExpanderArrow.h"
+
+template <> struct TListTypeTraits<FCurveEditorTreeItemID>
+{
+public:
+	struct NullableType : FCurveEditorTreeItemID
+	{
+		NullableType(TYPE_OF_NULLPTR){}
+		NullableType(FCurveEditorTreeItemID Other) : FCurveEditorTreeItemID(Other) {}
+	};
+
+	using MapKeyFuncs = TDefaultMapHashableKeyFuncs<FCurveEditorTreeItemID, TSharedRef<ITableRow>, false>;
+	using MapKeyFuncsSparse = TDefaultMapHashableKeyFuncs<FCurveEditorTreeItemID, FSparseItemInfo, false>;
+	using SetKeyFuncs = DefaultKeyFuncs<FCurveEditorTreeItemID>;
+
+	static void AddReferencedObjects(FReferenceCollector&, TArray<FCurveEditorTreeItemID>&, TSet<FCurveEditorTreeItemID>&) {}
+
+	static bool IsPtrValid(NullableType InPtr)
+	{
+		return InPtr.IsValid();
+	}
+
+	static void ResetPtr(NullableType& InPtr)
+	{
+		InPtr = nullptr;
+	}
+
+	static NullableType MakeNullPtr()
+	{
+		return nullptr;
+	}
+
+	static FCurveEditorTreeItemID NullableItemTypeConvertToItemType(NullableType InPtr)
+	{
+		return InPtr;
+	}
+
+	static FString DebugDump(FCurveEditorTreeItemID InPtr)
+	{
+		return FString::Printf(TEXT("%d"), InPtr.GetValue());
+	}
+
+	class SerializerType{};
+};
+template <>
+struct TIsValidListItem<FCurveEditorTreeItemID>
+{
+	enum
+	{
+		Value = true
+	};
+};
+
+class SCurveEditorTreeView : public STreeView<FCurveEditorTreeItemID>
+{};
+
+struct SCurveEditorTableRow : SMultiColumnTableRow<FCurveEditorTreeItemID>
+{
+	FCurveEditorTreeItemID TreeItemID;
+	TWeakPtr<FCurveEditor> WeakCurveEditor;
+
+	void Construct(const FTableRowArgs& InArgs, const TSharedRef<STableViewBase>& OwnerTableView, TWeakPtr<FCurveEditor> InCurveEditor, FCurveEditorTreeItemID InTreeItemID)
+	{
+		TreeItemID = InTreeItemID;
+		WeakCurveEditor = InCurveEditor;
+
+		SMultiColumnTableRow::Construct(InArgs, OwnerTableView);
+	}
+
+	virtual TSharedRef<SWidget> GenerateWidgetForColumn(const FName& InColumnName) override
+	{
+		TSharedPtr<FCurveEditor> CurveEditor = WeakCurveEditor.Pin();
+		TSharedPtr<ICurveEditorTreeItem> TreeItem = CurveEditor ? CurveEditor->GetTreeItem(TreeItemID).GetItem() : nullptr;
+
+		TSharedPtr<SWidget> Widget;
+		if (TreeItem.IsValid())
+		{
+			Widget = TreeItem->GenerateCurveEditorTreeWidget(InColumnName, WeakCurveEditor, TreeItemID);
+		}
+
+		if (!Widget.IsValid())
+		{
+			Widget = SNullWidget::NullWidget;
+		}
+
+		if (InColumnName == ICurveEditorTreeItem::ColumnNames.Label)
+		{
+			return SNew(SHorizontalBox)
+				+ SHorizontalBox::Slot()
+				.AutoWidth()
+				[
+					SNew(SExpanderArrow, SharedThis(this))
+				]
+
+				+ SHorizontalBox::Slot()
+				.Padding(FMargin(3.f, 0.f, 0.f, 0.f))
+				.VAlign(VAlign_Center)
+				[
+					Widget.ToSharedRef()
+				];
+		}
+
+		return Widget.ToSharedRef();
+	}
+};
+
+void SCurveEditorTree::Construct(const FArguments& InArgs, TSharedPtr<FCurveEditor> InCurveEditor)
+{
+	HeaderRow = SNew(SHeaderRow)
+		.Visibility(EVisibility::Collapsed)
+
+		+ SHeaderRow::Column(ICurveEditorTreeItem::ColumnNames.Label)
+
+		+ SHeaderRow::Column(ICurveEditorTreeItem::ColumnNames.PinHeader)
+		.FixedWidth(24.f);
+
+	CurveEditor = InCurveEditor;
+
+	ChildSlot
+	[
+		SAssignNew(TreeView, SCurveEditorTreeView)
+		.SelectionMode(ESelectionMode::Multi)
+		.HeaderRow(HeaderRow)
+		.HighlightParentNodesForSelection(true)
+		.TreeItemsSource(&InCurveEditor->GetRootTreeItems())
+		.OnGetChildren(this, &SCurveEditorTree::GetTreeItemChildren)
+		.OnGenerateRow(this, &SCurveEditorTree::GenerateRow)
+		.OnSetExpansionRecursive(this, &SCurveEditorTree::SetItemExpansionRecursive)
+		.OnSelectionChanged_Lambda(
+			[this](TListTypeTraits<FCurveEditorTreeItemID>::NullableType InItemID, ESelectInfo::Type Type)
+			{
+				this->OnTreeSelectionChanged(InItemID, Type);
+			}
+		)
+	];
+
+	CurveEditor->OnTreeChanged().AddSP(TreeView.ToSharedRef(), &SCurveEditorTreeView::RequestTreeRefresh);
+}
+
+FReply SCurveEditorTree::OnKeyDown(const FGeometry& MyGeometry, const FKeyEvent& InKeyEvent)
+{
+	if (InKeyEvent.GetKey() == EKeys::Escape)
+	{
+		TreeView->ClearSelection();
+		return FReply::Handled();
+	}
+
+	return FReply::Unhandled();
+}
+
+TSharedRef<ITableRow> SCurveEditorTree::GenerateRow(FCurveEditorTreeItemID ItemID, const TSharedRef<STableViewBase>& OwnerTable)
+{
+	return SNew(SCurveEditorTableRow, OwnerTable, CurveEditor, ItemID);
+}
+
+void SCurveEditorTree::GetTreeItemChildren(FCurveEditorTreeItemID Parent, TArray<FCurveEditorTreeItemID>& OutChildren)
+{
+	TArrayView<const FCurveEditorTreeItemID> Children = CurveEditor->GetTreeItem(Parent).GetChildren();
+	OutChildren.Append(Children.GetData(), Children.Num());
+}
+
+void SCurveEditorTree::OnTreeSelectionChanged(FCurveEditorTreeItemID, ESelectInfo::Type)
+{
+	CurveEditor->SetDirectTreeSelection(TreeView->GetSelectedItems());
+}
+
+void SCurveEditorTree::SetItemExpansionRecursive(FCurveEditorTreeItemID Model, bool bInExpansionState)
+{
+	if (Model.IsValid())
+	{
+		TreeView->SetItemExpansion(Model, bInExpansionState);
+
+		TArray<FCurveEditorTreeItemID> Children;
+		GetTreeItemChildren(Model, Children);
+
+		for (FCurveEditorTreeItemID Child : Children)
+		{
+			if (Child.IsValid())
+			{
+				SetItemExpansionRecursive(Child, bInExpansionState);
+			}
+		}
+	}
+}
\ No newline at end of file
diff --git a/Engine/Source/Editor/CurveEditor/Private/Tree/SCurveEditorTreePin.cpp b/Engine/Source/Editor/CurveEditor/Private/Tree/SCurveEditorTreePin.cpp
new file mode 100644
--- /dev/null
+++ b/Engine/Source/Editor/CurveEditor/Private/Tree/SCurveEditorTreePin.cpp
@@ -0,0 +1,122 @@
+// Copyright 1998-2019 Epic Games, Inc. All Rights Reserved.
+
+#include "Tree/SCurveEditorTreePin.h"
+#include "CurveEditor.h"
+#include "Algo/AllOf.h"
+
+#include "Widgets/Images/SImage.h"
+#include "Widgets/Input/SButton.h"
+
+#include "EditorStyleSet.h"
+
+
+void SCurveEditorTreePin::Construct(const FArguments& InArgs, TWeakPtr<FCurveEditor> InCurveEditor, FCurveEditorTreeItemID InTreeItemID)
+{
+	WeakCurveEditor = InCurveEditor;
+	TreeItemID = InTreeItemID;
+
+	ChildSlot
+	[
+		SNew(SButton)
+		.ButtonStyle(FEditorStyle::Get(), "HoverHintOnly")
+		.Visibility(this, &SCurveEditorTreePin::GetPinVisibility)
+		.OnClicked(this, &SCurveEditorTreePin::TogglePinned)
+		[
+			SNew(SImage)
+			.Image(this, &SCurveEditorTreePin::GetPinBrush)
+		]
+	];
+}
+
+FReply SCurveEditorTreePin::TogglePinned()
+{
+	TSharedPtr<FCurveEditor> CurveEditor = WeakCurveEditor.Pin();
+	if (CurveEditor)
+	{
+		if (IsPinnedRecursive(TreeItemID, CurveEditor.Get()))
+		{
+			UnpinRecursive(TreeItemID, CurveEditor.Get());
+		}
+		else
+		{
+			PinRecursive(TreeItemID, CurveEditor.Get());
+		}
+	}
+	return FReply::Handled();
+}
+
+void SCurveEditorTreePin::PinRecursive(FCurveEditorTreeItemID InTreeItem, FCurveEditor* CurveEditor) const
+{
+	FCurveEditorTreeItem& Item = CurveEditor->GetTreeItem(InTreeItem);
+
+	for (FCurveModelID CurveID : Item.GetOrCreateCurves(CurveEditor))
+	{
+		CurveEditor->PinCurve(CurveID);
+	}
+
+	for (FCurveEditorTreeItemID Child : Item.GetChildren())
+	{
+		PinRecursive(Child, CurveEditor);
+	}
+}
+
+void SCurveEditorTreePin::UnpinRecursive(FCurveEditorTreeItemID InTreeItem, FCurveEditor* CurveEditor) const
+{
+	const bool bIsSelected = CurveEditor->GetTreeSelectionState(InTreeItem) == ECurveEditorTreeSelectionState::Explicit;
+
+	FCurveEditorTreeItem& Item = CurveEditor->GetTreeItem(InTreeItem);
+	for (FCurveModelID CurveID : Item.GetCurves())
+	{
+		if (bIsSelected)
+		{
+			CurveEditor->UnpinCurve(CurveID);
+		}
+		else
+		{
+			Item.DestroyCurves(CurveEditor);
+		}
+	}
+
+	for (FCurveEditorTreeItemID Child : Item.GetChildren())
+	{
+		UnpinRecursive(Child, CurveEditor);
+	}
+}
+
+bool SCurveEditorTreePin::IsPinnedRecursive(FCurveEditorTreeItemID InTreeItem, FCurveEditor* CurveEditor) const
+{
+	const FCurveEditorTreeItem& Item = CurveEditor->GetTreeItem(InTreeItem);
+
+	TArrayView<const FCurveModelID>          Curves   = Item.GetCurves();
+	TArrayView<const FCurveEditorTreeItemID> Children = Item.GetChildren();
+
+	const bool bAllChildren = Algo::AllOf(Children, [this, CurveEditor](FCurveEditorTreeItemID In){ return this->IsPinnedRecursive(In, CurveEditor); });
+
+	if (Curves.Num() == 0)
+	{
+		return Children.Num() > 0 && bAllChildren;
+	}
+
+	const bool bAllCurves = Algo::AllOf(Curves, [CurveEditor](FCurveModelID In) { return CurveEditor->IsCurvePinned(In); });
+	return bAllCurves && bAllChildren;
+}
+
+EVisibility SCurveEditorTreePin::GetPinVisibility() const
+{
+	// ToDo: These should only be visible when hovering over the owning widget row.
+	return EVisibility::Visible;
+}
+
+const FSlateBrush* SCurveEditorTreePin::GetPinBrush() const
+{
+	TSharedPtr<FCurveEditor> CurveEditor = WeakCurveEditor.Pin();
+	if (CurveEditor)
+	{
+		if (IsPinnedRecursive(TreeItemID, CurveEditor.Get()))
+		{
+			return FEditorStyle::GetBrush("GenericCurveEditor.Pin_Active");
+		}
+	}
+
+	return FEditorStyle::GetBrush("GenericCurveEditor.Pin_Inactive");
+}
\ No newline at end of file
diff --git a/Engine/Source/Editor/CurveEditor/Private/Views/SCurveEditorViewAbsolute.h b/Engine/Source/Editor/CurveEditor/Private/Views/SCurveEditorViewAbsolute.h
new file mode 100644
--- /dev/null
+++ b/Engine/Source/Editor/CurveEditor/Private/Views/SCurveEditorViewAbsolute.h
@@ -0,0 +1,16 @@
+// Copyright 1998-2018 Epic Games, Inc. All Rights Reserved.
+
+#pragma once
+
+#include "Views/SInteractiveCurveEditorView.h"
+
+/**
+ * A Normalized curve view supporting one or more curves with their own screen transform that normalizes the vertical curve range to [-1,1]
+ */
+class SCurveEditorViewAbsolute : public SInteractiveCurveEditorView
+{
+public:
+
+	void Construct(const FArguments& InArgs, TWeakPtr<FCurveEditor> InCurveEditor);
+
+};
diff --git a/Engine/Source/Editor/CurveEditor/Private/Views/SCurveEditorViewNormalized.cpp b/Engine/Source/Editor/CurveEditor/Private/Views/SCurveEditorViewNormalized.cpp
new file mode 100644
--- /dev/null
+++ b/Engine/Source/Editor/CurveEditor/Private/Views/SCurveEditorViewNormalized.cpp
@@ -0,0 +1,83 @@
+// Copyright 1998-2018 Epic Games, Inc. All Rights Reserved.
+
+#include "Views/SCurveEditorViewNormalized.h"
+#include "CurveEditor.h"
+#include "CurveEditorHelpers.h"
+#include "CurveModel.h"
+#include "Widgets/Text/STextBlock.h"
+
+constexpr float NormalizedPadding = 10.f;
+
+void SCurveEditorViewNormalized::Construct(const FArguments& InArgs, TWeakPtr<FCurveEditor> InCurveEditor)
+{
+	bFixedOutputBounds = true;
+	OutputMin = 0.0;
+	OutputMax = 1.0;
+
+	SInteractiveCurveEditorView::Construct(InArgs, InCurveEditor);
+
+	ChildSlot
+	.HAlign(HAlign_Right)
+	.VAlign(VAlign_Top)
+	.Padding(FMargin(0.f, CurveViewConstants::CurveLabelOffsetY, CurveViewConstants::CurveLabelOffsetX, 0.f))
+	[
+		SNew(STextBlock)
+		.Font(FCoreStyle::Get().GetFontStyle("ToolTip.LargerFont"))
+		.ColorAndOpacity(this, &SCurveEditorViewNormalized::GetCurveCaptionColor)
+		.Text(this, &SCurveEditorViewNormalized::GetCurveCaption)
+	];
+}
+
+void SCurveEditorViewNormalized::GetGridLinesY(TSharedRef<FCurveEditor> CurveEditor, TArray<float>& MajorGridLines, TArray<float>& MinorGridLines, TArray<FText>& MajorGridLabels) const
+{
+	FCurveEditorScreenSpace ViewSpace = GetViewSpace();
+
+	MajorGridLines.Add(ViewSpace.ValueToScreen(0.0));
+	MajorGridLines.Add(ViewSpace.ValueToScreen(0.5));
+	MajorGridLines.Add(ViewSpace.ValueToScreen(1.0));
+
+	MinorGridLines.Add(ViewSpace.ValueToScreen(0.25));
+	MinorGridLines.Add(ViewSpace.ValueToScreen(0.75));
+}
+
+void SCurveEditorViewNormalized::Tick(const FGeometry& AllottedGeometry, const double InCurrentTime, const float InDeltaTime)
+{
+	TSharedPtr<FCurveEditor> CurveEditor = WeakCurveEditor.Pin();
+	if (!CurveEditor)
+	{
+		return;
+	}
+
+	if (!CurveEditor->AreBoundTransformUpdatesSuppressed())
+	{
+		const double ValuePerPixel = 1.0 / AllottedGeometry.GetLocalSize().Y;
+		const double ValueSpacePadding = NormalizedPadding * ValuePerPixel;
+
+		OutputMin = 0.0 - ValueSpacePadding;
+		OutputMax = 1.0 + ValueSpacePadding;
+
+		for (auto It = CurveInfoByID.CreateIterator(); It; ++It)
+		{
+			FCurveModel* Curve = CurveEditor->FindCurve(It.Key());
+			if (!ensureAlways(Curve))
+			{
+				continue;
+			}
+
+			double CurveOutputMin = 0, CurveOutputMax = 1;
+			Curve->GetValueRange(CurveOutputMin, CurveOutputMax);
+
+			const double Scale = (CurveOutputMax - CurveOutputMin);
+			if (CurveOutputMax > CurveOutputMin)
+			{
+				It->Value.ViewToCurveTransform = FTransform2D(FScale2D(1.f, Scale), FVector2D(0.f, CurveOutputMin));
+			}
+			else
+			{
+				It->Value.ViewToCurveTransform = FVector2D(0.f, CurveOutputMin-0.5);
+			}
+		}
+	}
+
+	SInteractiveCurveEditorView::Tick(AllottedGeometry, InCurrentTime, InDeltaTime);
+}
\ No newline at end of file
diff --git a/Engine/Source/Editor/CurveEditor/Private/Views/SCurveEditorViewNormalized.h b/Engine/Source/Editor/CurveEditor/Private/Views/SCurveEditorViewNormalized.h
new file mode 100644
--- /dev/null
+++ b/Engine/Source/Editor/CurveEditor/Private/Views/SCurveEditorViewNormalized.h
@@ -0,0 +1,24 @@
+// Copyright 1998-2018 Epic Games, Inc. All Rights Reserved.
+
+#pragma once
+
+#include "Views/SInteractiveCurveEditorView.h"
+
+/**
+ * A Normalized curve view supporting one or more curves with their own screen transform that normalizes the vertical curve range to [-1,1]
+ */
+class SCurveEditorViewNormalized : public SInteractiveCurveEditorView
+{
+public:
+
+	void Construct(const FArguments& InArgs, TWeakPtr<FCurveEditor> InCurveEditor);
+
+	/** Tools should ignore vertical snapping because it causes issues with curves that have tiny extents. */
+	virtual bool IsValueSnapEnabled() const override { return false; }
+
+private:
+
+	virtual void Tick(const FGeometry& AllottedGeometry, const double InCurrentTime, const float InDeltaTime) override;
+
+	virtual void GetGridLinesY(TSharedRef<FCurveEditor> CurveEditor, TArray<float>& MajorGridLines, TArray<float>& MinorGridLines, TArray<FText>& MajorGridLabels) const override;
+};
diff --git a/Engine/Source/Editor/CurveEditor/Private/Views/SCurveEditorViewStacked.h b/Engine/Source/Editor/CurveEditor/Private/Views/SCurveEditorViewStacked.h
new file mode 100644
--- /dev/null
+++ b/Engine/Source/Editor/CurveEditor/Private/Views/SCurveEditorViewStacked.h
@@ -0,0 +1,29 @@
+// Copyright 1998-2018 Epic Games, Inc. All Rights Reserved.
+
+#pragma once
+
+#include "Views/SInteractiveCurveEditorView.h"
+
+/**
+ * A Normalized curve view supporting one or more curves with their own screen transform that normalizes the vertical curve range to [-1,1]
+ */
+class SCurveEditorViewStacked : public SInteractiveCurveEditorView
+{
+public:
+
+	void Construct(const FArguments& InArgs, TWeakPtr<FCurveEditor> InCurveEditor);
+	
+	/** Tools should ignore vertical snapping because it causes issues with curves that have tiny extents. */
+	virtual bool IsValueSnapEnabled() const override { return false; }
+
+private:
+
+	virtual void Tick(const FGeometry& AllottedGeometry, const double InCurrentTime, const float InDeltaTime) override;
+	virtual FVector2D ComputeDesiredSize(float LayoutScaleMultiplier) const override;
+
+	virtual void PaintView(const FPaintArgs& Args, const FGeometry& AllottedGeometry, const FSlateRect& MyCullingRect, FSlateWindowElementList& OutDrawElements, int32 BaseLayerId, const FWidgetStyle& InWidgetStyle, bool bParentEnabled) const override;
+	virtual void GetGridLinesY(TSharedRef<FCurveEditor> CurveEditor, TArray<float>& MajorGridLines, TArray<float>& MinorGridLines, TArray<FText>& MajorGridLabels) const override;
+
+	void DrawViewGrids(const FGeometry& AllottedGeometry, const FSlateRect& MyCullingRect, FSlateWindowElementList& OutDrawElements, int32 BaseLayerId, ESlateDrawEffect DrawEffects) const;
+	void DrawLabels(const FGeometry& AllottedGeometry, const FSlateRect& MyCullingRect, FSlateWindowElementList& OutDrawElements, int32 BaseLayerId, ESlateDrawEffect DrawEffects) const;
+};
