This is a significant overhaul to the Curve Editor used by Sequencer which adds a plugin-based architecture and extensibility. New tools and toolbar buttons can be added to all usages of the curve edi

diff --git a/Engine/Source/Editor/CurveEditor/Private/Views/SInteractiveCurveEditorView.cpp b/Engine/Source/Editor/CurveEditor/Private/Views/SInteractiveCurveEditorView.cpp
new file mode 100644
--- /dev/null
+++ b/Engine/Source/Editor/CurveEditor/Private/Views/SInteractiveCurveEditorView.cpp
@@ -0,0 +1,1376 @@
+// Copyright 1998-2018 Epic Games, Inc. All Rights Reserved.
+
+#include "Views/SInteractiveCurveEditorView.h"
+#include "EditorStyleSet.h"
+#include "Rendering/DrawElements.h"
+#include "Styling/CoreStyle.h"
+#include "CurveEditor.h"
+#include "CurveEditorScreenSpace.h"
+#include "CurveDrawInfo.h"
+#include "CurveEditorSettings.h"
+#include "ICurveEditorBounds.h"
+#include "Types/SlateStructs.h"
+#include "Fonts/FontMeasure.h"
+#include "Widgets/SToolTip.h"
+#include "Widgets/Text/STextBlock.h"
+#include "DragOperations/CurveEditorDragOperation_Tangent.h"
+#include "DragOperations/CurveEditorDragOperation_MoveKeys.h"
+#include "DragOperations/CurveEditorDragOperation_Marquee.h"
+#include "DragOperations/CurveEditorDragOperation_Pan.h"
+#include "DragOperations/CurveEditorDragOperation_Zoom.h"
+#include "SCurveEditorPanel.h"
+#include "Framework/MultiBox/MultiBoxBuilder.h"
+#include "CurveEditorContextMenu.h"
+#include "CurveEditorCommands.h"
+#include "CurveEditorHelpers.h"
+#include "CurveEditor.h"
+#include "ITimeSlider.h"
+
+namespace CurveViewConstants
+{
+	/** The number of pixels to offset Labels from the Left/Right size. */
+	constexpr float LabelOffsetPixels = 2.f;
+
+	/** The number of pixels away the mouse can be and still be considering hovering over a curve. */
+	constexpr float HoverProximityThresholdPx = 5.f;
+}
+
+#define LOCTEXT_NAMESPACE "SInteractiveCurveEditorView"
+
+TUniquePtr<ICurveEditorKeyDragOperation> CreateKeyDrag(ECurvePointType KeyType)
+{
+	switch (KeyType)
+	{
+	case ECurvePointType::ArriveTangent:
+	case ECurvePointType::LeaveTangent:
+		return MakeUnique<FCurveEditorDragOperation_Tangent>();
+
+	default:
+		return MakeUnique<FCurveEditorDragOperation_MoveKeys>();
+	}
+}
+
+class SDynamicToolTip : public SToolTip
+{
+public:
+	TAttribute<bool> bIsEnabled;
+	virtual bool IsEmpty() const override { return !bIsEnabled.Get(); }
+};
+
+void SInteractiveCurveEditorView::Construct(const FArguments& InArgs, TWeakPtr<FCurveEditor> InCurveEditor)
+{
+	FixedHeight = InArgs._FixedHeight;
+	BackgroundTint = InArgs._BackgroundTint;
+	MaximumCapacity = InArgs._MaximumCapacity;
+	bAutoSize = InArgs._AutoSize;
+
+	WeakCurveEditor = InCurveEditor;
+
+	InCurveEditor.Pin()->OnActiveToolChangedDelegate.AddSP(this, &SInteractiveCurveEditorView::OnCurveEditorToolChanged);
+
+	TSharedRef<SDynamicToolTip> ToolTipWidget =
+		SNew(SDynamicToolTip)
+		.BorderImage(FCoreStyle::Get().GetBrush("ToolTip.BrightBackground"))
+		[
+			SNew(SVerticalBox)
+			+ SVerticalBox::Slot()
+			[
+				SNew(STextBlock)
+				.Text(this, &SInteractiveCurveEditorView::GetToolTipCurveName)
+				.Font(FCoreStyle::Get().GetFontStyle("ToolTip.LargerFont"))
+				.ColorAndOpacity(FLinearColor::Black)
+			]
+		+ SVerticalBox::Slot()
+		[
+			SNew(STextBlock)
+			.Text(this, &SInteractiveCurveEditorView::GetToolTipTimeText)
+			.Font(FCoreStyle::Get().GetFontStyle("ToolTip.LargerFont"))
+			.ColorAndOpacity(FLinearColor::Black)
+		]
+		+ SVerticalBox::Slot()
+		[
+			SNew(STextBlock)
+			.Text(this, &SInteractiveCurveEditorView::GetToolTipValueText)
+			.Font(FCoreStyle::Get().GetFontStyle("ToolTip.LargerFont"))
+			.ColorAndOpacity(FLinearColor::Black)
+		]
+	];
+
+	ToolTipWidget->bIsEnabled = MakeAttributeSP(this, &SInteractiveCurveEditorView::IsToolTipEnabled);
+	SetToolTip(ToolTipWidget);
+}
+
+FText SInteractiveCurveEditorView::GetCurveCaption() const
+{
+	TSharedPtr<FCurveEditor> CurveEditor = WeakCurveEditor.Pin();
+	if (CurveEditor && CurveInfoByID.Num() == 1)
+	{
+		for (const TTuple<FCurveModelID, FCurveInfo>& Pair : CurveInfoByID)
+		{
+			if (const FCurveModel* Curve = CurveEditor->FindCurve(Pair.Key))
+			{
+				return Curve->GetLongDisplayName();
+			}
+		}
+	}
+
+	return FText::GetEmpty();
+}
+
+FSlateColor SInteractiveCurveEditorView::GetCurveCaptionColor() const
+{
+	TSharedPtr<FCurveEditor> CurveEditor = WeakCurveEditor.Pin();
+	if (CurveEditor && CurveInfoByID.Num() == 1)
+	{
+		for (const TTuple<FCurveModelID, FCurveInfo>& Pair : CurveInfoByID)
+		{
+			if (const FCurveModel* Curve = CurveEditor->FindCurve(Pair.Key))
+			{
+				return Curve->GetColor();
+			}
+		}
+	}
+
+	return BackgroundTint.CopyWithNewOpacity(1.f);
+}
+
+void SInteractiveCurveEditorView::GetGridLinesX(TSharedRef<FCurveEditor> CurveEditor, TArray<float>& MajorGridLines, TArray<float>& MinorGridLines, TArray<FText>& MajorGridLabels) const
+{
+	CurveEditor->GetGridLinesX(MajorGridLines, MinorGridLines, MajorGridLabels);
+
+	FCurveEditorScreenSpaceH PanelSpace = CurveEditor->GetPanelInputSpace();
+	FCurveEditorScreenSpaceH ViewSpace  = GetViewSpace();
+
+	double InputOffset = ViewSpace.GetInputMin() - PanelSpace.GetInputMin();
+	if (InputOffset != 0.0)
+	{
+		const float PixelDifference = InputOffset * PanelSpace.PixelsPerInput();
+		for (float& Line : MajorGridLines)
+		{
+			Line -= PixelDifference;
+		}
+		for (float& Line : MinorGridLines)
+		{
+			Line -= PixelDifference;
+		}
+	}
+}
+
+void SInteractiveCurveEditorView::GetGridLinesY(TSharedRef<FCurveEditor> CurveEditor, TArray<float>& MajorGridLines, TArray<float>& MinorGridLines, TArray<FText>& MajorGridLabels) const
+{
+	CurveEditor::ConstructYGridLines(GetViewSpace(), 4, MajorGridLines, MinorGridLines, CurveEditor->GetGridLineLabelFormatYAttribute().Get(), &MajorGridLabels);
+}
+
+int32 SInteractiveCurveEditorView::OnPaint(const FPaintArgs& Args, const FGeometry& AllottedGeometry, const FSlateRect& MyCullingRect, FSlateWindowElementList& OutDrawElements, int32 BaseLayerId, const FWidgetStyle& InWidgetStyle, bool bParentEnabled) const
+{
+	PaintView(Args, AllottedGeometry, MyCullingRect, OutDrawElements, BaseLayerId, InWidgetStyle, bParentEnabled);
+	SCompoundWidget::OnPaint(Args, AllottedGeometry, MyCullingRect, OutDrawElements, BaseLayerId + CurveViewConstants::ELayerOffset::WidgetContent, InWidgetStyle, bParentEnabled);
+
+	return BaseLayerId;
+}
+
+void SInteractiveCurveEditorView::PaintView(const FPaintArgs& Args, const FGeometry& AllottedGeometry, const FSlateRect& MyCullingRect, FSlateWindowElementList& OutDrawElements, int32 BaseLayerId, const FWidgetStyle& InWidgetStyle, bool bParentEnabled) const
+{
+	TSharedPtr<FCurveEditor> CurveEditor = WeakCurveEditor.Pin();
+	if (CurveEditor)
+	{
+		const ESlateDrawEffect DrawEffects = ShouldBeEnabled(bParentEnabled) ? ESlateDrawEffect::None : ESlateDrawEffect::DisabledEffect;
+
+		DrawBackground(AllottedGeometry, OutDrawElements, BaseLayerId, DrawEffects);
+		DrawGridLines(CurveEditor.ToSharedRef(), AllottedGeometry, OutDrawElements, BaseLayerId, DrawEffects);
+		DrawCurves(CurveEditor.ToSharedRef(), AllottedGeometry, MyCullingRect, OutDrawElements, BaseLayerId, InWidgetStyle, DrawEffects);
+	}
+}
+
+void SInteractiveCurveEditorView::DrawBackground(const FGeometry& AllottedGeometry, FSlateWindowElementList& OutDrawElements, int32 BaseLayerId, ESlateDrawEffect DrawEffects) const
+{
+	if (BackgroundTint != FLinearColor::White)
+	{
+		FSlateDrawElement::MakeBox(OutDrawElements, BaseLayerId + CurveViewConstants::ELayerOffset::Background, AllottedGeometry.ToPaintGeometry(),
+			FEditorStyle::GetBrush("ToolPanel.GroupBorder"), DrawEffects, BackgroundTint);
+	}
+}
+
+void SInteractiveCurveEditorView::DrawGridLines(TSharedRef<FCurveEditor> CurveEditor, const FGeometry& AllottedGeometry, FSlateWindowElementList& OutDrawElements, int32 BaseLayerId, ESlateDrawEffect DrawEffects) const
+{
+	// Rendering info
+	const float          Width = AllottedGeometry.GetLocalSize().X;
+	const float          Height = AllottedGeometry.GetLocalSize().Y;
+	const float          RoundedWidth = FMath::RoundToFloat(Width);
+	const float          RoundedHeight = FMath::RoundToFloat(Height);
+	const FLinearColor   MajorGridColor = CurveEditor->GetPanel()->GetGridLineTint();
+	const FLinearColor   MinorGridColor = MajorGridColor.CopyWithNewOpacity(MajorGridColor.A * .5f);
+	const FPaintGeometry PaintGeometry = AllottedGeometry.ToPaintGeometry();
+	const FLinearColor	 LabelColor = FLinearColor::White.CopyWithNewOpacity(0.65f);
+	const FSlateFontInfo FontInfo = FCoreStyle::Get().GetFontStyle("ToolTip.LargerFont");
+
+	// Get our viewing range bounds. We go through the GetBounds() interface on the curve editor because it's more aware of what our range is than the original widget is.
+	double InputValueMin, InputValueMax;
+	GetInputBounds(InputValueMin, InputValueMax);
+
+	TArray<float> MajorGridLines, MinorGridLines;
+	TArray<FText> MajorGridLabels;
+
+	GetGridLinesX(CurveEditor, MajorGridLines, MinorGridLines, MajorGridLabels);
+	ensureMsgf(MajorGridLabels.Num() == 0 || MajorGridLines.Num() == MajorGridLabels.Num(), TEXT("If grid labels are specified, one must be specified for every major grid line, even if it is just an empty FText."));
+
+	// Pre-allocate an array of line points to draw our vertical lines. Each major grid line
+	// will overwrite the X value of both points but leave the Y value untouched so they draw from the bottom to the top.
+	TArray<FVector2D> LinePoints;
+	LinePoints.Add(FVector2D(0.f, 0.f));
+	LinePoints.Add(FVector2D(0.f, Height));
+
+	const TSharedRef<FSlateFontMeasure> FontMeasure = FSlateApplication::Get().GetRenderer()->GetFontMeasureService();
+
+	// Draw major vertical grid lines
+	for (int32 i = 0; i < MajorGridLines.Num(); i++)
+	{
+		const float RoundedLine = FMath::RoundToFloat(MajorGridLines[i]);
+		if (RoundedLine < 0 || RoundedLine > RoundedWidth)
+		{
+			continue;
+		}
+
+		// Vertical Grid Line
+		LinePoints[0].X = LinePoints[1].X = RoundedLine;
+
+		if (MajorGridLabels.IsValidIndex(i))
+		{
+			FText Label = MajorGridLabels[i];
+			const FVector2D LabelSize = FontMeasure->Measure(Label, FontInfo);
+			const FPaintGeometry LabelGeometry = AllottedGeometry.ToPaintGeometry(FSlateLayoutTransform(FVector2D(LinePoints[0].X - LabelSize.X*.5f, CurveViewConstants::LabelOffsetPixels)));
+
+			LinePoints[0].Y = LabelSize.Y + CurveViewConstants::LabelOffsetPixels*2.f;
+
+			FSlateDrawElement::MakeText(
+				OutDrawElements,
+				BaseLayerId + CurveViewConstants::ELayerOffset::GridLabels,
+				LabelGeometry,
+				Label,
+				FontInfo,
+				DrawEffects,
+				LabelColor
+			);
+		}
+		else
+		{
+			LinePoints[0].Y = 0.f;
+		}
+
+		FSlateDrawElement::MakeLines(
+			OutDrawElements,
+			BaseLayerId + CurveViewConstants::ELayerOffset::GridLines,
+			PaintGeometry,
+			LinePoints,
+			DrawEffects,
+			MajorGridColor,
+			false
+		);
+	}
+
+	LinePoints[0].Y = 0.f;
+
+	// Now draw the minor vertical lines which are drawn with a lighter color.
+	for (float PosX : MinorGridLines)
+	{
+		if (PosX < 0 || PosX > Width)
+		{
+			continue;
+		}
+
+		LinePoints[0].X = LinePoints[1].X = PosX;
+
+		FSlateDrawElement::MakeLines(
+			OutDrawElements,
+			BaseLayerId + CurveViewConstants::ELayerOffset::GridLines,
+			PaintGeometry,
+			LinePoints,
+			DrawEffects,
+			MinorGridColor,
+			false
+		);
+	}
+
+	MajorGridLines.Reset();
+	MinorGridLines.Reset();
+	MajorGridLabels.Reset();
+	GetGridLinesY(CurveEditor, MajorGridLines, MinorGridLines, MajorGridLabels);
+	ensureMsgf(MajorGridLabels.Num() == 0 || MajorGridLines.Num() == MajorGridLabels.Num(), TEXT("If grid labels are specified, one must be specified for every major grid line, even if it is just an empty FText."));
+
+	// Reset our cached Line to draw from left to right
+	LinePoints[0].X = 0.f;
+	LinePoints[1].X = Width;
+
+	// Draw our major horizontal lines
+	for (int32 i = 0; i < MajorGridLines.Num(); i++)
+	{
+		const float RoundedLine = FMath::RoundToFloat(MajorGridLines[i]);
+		if (RoundedLine < 0 || RoundedLine > RoundedHeight)
+		{
+			continue;
+		}
+
+		// Overwrite the height of the line we're drawing to draw the different grid lines.
+		LinePoints[0].Y = LinePoints[1].Y = RoundedLine;
+
+		if (MajorGridLabels.IsValidIndex(i))
+		{
+			FText Label = MajorGridLabels[i];
+			const FVector2D LabelSize = FontMeasure->Measure(Label, FontInfo);
+			const FPaintGeometry LabelGeometry = AllottedGeometry.ToPaintGeometry(FSlateLayoutTransform(FVector2D(CurveViewConstants::LabelOffsetPixels, LinePoints[0].Y - LabelSize.Y*.5f)));
+
+			LinePoints[0].X = LabelSize.X + CurveViewConstants::LabelOffsetPixels*2.f;
+
+			FSlateDrawElement::MakeText(
+				OutDrawElements,
+				BaseLayerId + CurveViewConstants::ELayerOffset::GridLabels,
+				LabelGeometry,
+				Label,
+				FontInfo,
+				DrawEffects,
+				LabelColor
+			);
+		}
+		else
+		{
+			LinePoints[0].X = 0.f;
+		}
+
+		FSlateDrawElement::MakeLines(
+			OutDrawElements,
+			BaseLayerId + CurveViewConstants::ELayerOffset::GridLines,
+			PaintGeometry,
+			LinePoints,
+			DrawEffects,
+			MajorGridColor,
+			false
+		);
+	}
+
+	LinePoints[0].X = 0.f;
+
+	// Draw our minor horizontal lines
+	for (float PosY : MinorGridLines)
+	{
+		if (PosY < 0 || PosY > Height)
+		{
+			continue;
+		}
+
+		LinePoints[0].Y = LinePoints[1].Y = PosY;
+
+		// Now draw the minor grid lines with a lighter color.
+		FSlateDrawElement::MakeLines(
+			OutDrawElements,
+			BaseLayerId + CurveViewConstants::ELayerOffset::GridLines,
+			PaintGeometry,
+			LinePoints,
+			DrawEffects,
+			MinorGridColor,
+			false
+		);
+	}
+}
+
+void SInteractiveCurveEditorView::DrawCurves(TSharedRef<FCurveEditor> CurveEditor, const FGeometry& AllottedGeometry, const FSlateRect& MyCullingRect, FSlateWindowElementList& OutDrawElements, int32 BaseLayerId, const FWidgetStyle& InWidgetStyle, ESlateDrawEffect DrawEffects) const
+{
+	static const FName SelectionColorName("SelectionColor");
+	FLinearColor SelectionColor = FEditorStyle::GetSlateColor(SelectionColorName).GetColor(InWidgetStyle);
+
+	const FVector2D      VisibleSize = AllottedGeometry.GetLocalSize();
+	const FPaintGeometry PaintGeometry = AllottedGeometry.ToPaintGeometry();
+
+	const float HoveredCurveThickness = 3.f;
+	const float UnHoveredCurveThickness = 1.f;
+	const bool  bAntiAliasCurves = true;
+
+	TOptional<FCurveModelID> HoveredCurve = GetHoveredCurve();
+	for (const FCurveDrawParams& Params : CachedDrawParams)
+	{
+		const bool bIsCurveHovered = HoveredCurve.IsSet() && HoveredCurve.GetValue() == Params.GetID();
+		const float Thickness = bIsCurveHovered ? HoveredCurveThickness : UnHoveredCurveThickness;
+		const int32 CurveLayerId = bIsCurveHovered ? BaseLayerId + CurveViewConstants::ELayerOffset::Curves : BaseLayerId + CurveViewConstants::ELayerOffset::HoveredCurves;
+
+		FSlateDrawElement::MakeLines(
+			OutDrawElements,
+			CurveLayerId,
+			PaintGeometry,
+			Params.InterpolatingPoints,
+			DrawEffects,
+			Params.Color,
+			bAntiAliasCurves,
+			Thickness
+		);
+
+		TArray<FVector2D> LinePoints;
+		LinePoints.SetNum(2);
+
+		// Draw tangents
+		for(int32 PointIndex = 0; PointIndex < Params.Points.Num(); PointIndex++)
+		{
+			const FCurvePointInfo& Point = Params.Points[PointIndex];
+			const FKeyDrawInfo& PointDrawInfo = Params.GetKeyDrawInfo(Point.Type, PointIndex);
+			const bool          bSelected = CurveEditor->GetSelection().IsSelected(FCurvePointHandle(Params.GetID(), Point.Type, Point.KeyHandle));
+			const FLinearColor  PointTint = bSelected ? SelectionColor : PointDrawInfo.Tint;
+
+			const int32 KeyLayerId = BaseLayerId + Point.LayerBias + (bSelected ? CurveViewConstants::ELayerOffset::SelectedKeys : CurveViewConstants::ELayerOffset::Keys);
+
+			if (Point.LineDelta.X != 0.f || Point.LineDelta.Y != 0.f)
+			{
+				LinePoints[0] = Point.ScreenPosition + Point.LineDelta.GetSafeNormal() * (PointDrawInfo.ScreenSize.X*.5f);
+				LinePoints[1] = Point.ScreenPosition + Point.LineDelta;
+
+				// Draw the connecting line - connecting lines are always drawn below everything else
+				FSlateDrawElement::MakeLines(OutDrawElements, BaseLayerId + CurveViewConstants::ELayerOffset::Keys - 1, PaintGeometry, LinePoints, DrawEffects, PointTint, true);
+			}
+
+			FPaintGeometry PointGeometry = AllottedGeometry.ToPaintGeometry(
+				Point.ScreenPosition - (PointDrawInfo.ScreenSize * 0.5f),
+				PointDrawInfo.ScreenSize
+			);
+
+			FSlateDrawElement::MakeBox(OutDrawElements, KeyLayerId, PointGeometry, PointDrawInfo.Brush, DrawEffects, PointTint);
+		}
+	}
+}
+
+void SInteractiveCurveEditorView::Tick(const FGeometry& AllottedGeometry, const double InCurrentTime, const float InDeltaTime)
+{
+	// Cache our Curve Drawing Params. These are used in multiple places so we cache them once each frame.
+	CachedDrawParams.Reset();
+	GetCurveDrawParams(CachedDrawParams);
+}
+
+void SInteractiveCurveEditorView::GetCurveDrawParams(TArray<FCurveDrawParams>& OutDrawParams) const
+{
+	TSharedPtr<FCurveEditor> CurveEditor = WeakCurveEditor.Pin();
+	if (!CurveEditor)
+	{
+		return;
+	}
+
+	// Get the Min/Max values on the X axis, for Time
+	double InputMin = 0, InputMax = 1;
+	GetInputBounds(InputMin, InputMax);
+
+	ECurveEditorTangentVisibility TangentVisibility = CurveEditor->GetSettings()->GetTangentVisibility();
+	OutDrawParams.Reserve(CurveInfoByID.Num());
+
+	for (const TTuple<FCurveModelID, FCurveInfo>& Pair : CurveInfoByID)
+	{
+		FCurveModel* CurveModel = CurveEditor->FindCurve(Pair.Key);
+		if (!ensureAlways(CurveModel))
+		{
+			continue;
+		}
+
+		FCurveEditorScreenSpace CurveSpace = GetCurveSpace(Pair.Key);
+		
+		const float DisplayRatio = (CurveSpace.PixelsPerOutput() / CurveSpace.PixelsPerInput());
+
+		const FKeyHandleSet* SelectedKeys = CurveEditor->GetSelection().GetAll().Find(Pair.Key);
+
+		// Create a new set of Curve Drawing Parameters to represent this particular Curve
+		FCurveDrawParams Params(Pair.Key);
+		Params.Color = CurveModel->GetColor();
+
+		// Gather the display metrics to use for each key type. This allows a Curve Model to override
+		// whether or not the curve supports Keys, Arrive/Leave Tangents, etc. If the Curve Model doesn't
+		// support a particular capability we can skip drawing them.
+		CurveModel->GetKeyDrawInfo(ECurvePointType::ArriveTangent, FKeyHandle::Invalid(), Params.ArriveTangentDrawInfo);
+		CurveModel->GetKeyDrawInfo(ECurvePointType::LeaveTangent, FKeyHandle::Invalid(), Params.LeaveTangentDrawInfo);
+
+		// Gather the interpolating points in input/output space
+		TArray<TTuple<double, double>> InterpolatingPoints;
+
+		CurveModel->DrawCurve(*CurveEditor, CurveSpace, InterpolatingPoints);
+		Params.InterpolatingPoints.Reserve(InterpolatingPoints.Num());
+
+		// An Input Offset allows for a fixed offset to all keys, such as displaying them in the middle of a frame instead of at the start.
+		double InputOffset = CurveModel->GetInputDisplayOffset();
+
+		// Convert the interpolating points to screen space
+		for (TTuple<double, double> Point : InterpolatingPoints)
+		{
+			Params.InterpolatingPoints.Add(
+				FVector2D(
+					CurveSpace.SecondsToScreen(Point.Get<0>() + InputOffset),
+					CurveSpace.ValueToScreen(Point.Get<1>())
+				)
+			);
+		}
+
+		TArray<FKeyHandle> VisibleKeys;
+		CurveModel->GetKeys(*CurveEditor, InputMin, InputMax, TNumericLimits<double>::Lowest(), TNumericLimits<double>::Max(), VisibleKeys);
+
+		if (VisibleKeys.Num())
+		{
+			TArray<FKeyPosition> AllKeyPositions;
+			TArray<FKeyAttributes> AllKeyAttributes;
+
+			AllKeyPositions.SetNum(VisibleKeys.Num());
+			AllKeyAttributes.SetNum(VisibleKeys.Num());
+
+			CurveModel->GetKeyPositions(VisibleKeys, AllKeyPositions);
+			CurveModel->GetKeyAttributes(VisibleKeys, AllKeyAttributes);
+
+			for (int32 Index = 0; Index < VisibleKeys.Num(); ++Index)
+			{
+				const FKeyHandle      KeyHandle = VisibleKeys[Index];
+				const FKeyPosition&   KeyPosition = AllKeyPositions[Index];
+				const FKeyAttributes& Attributes = AllKeyAttributes[Index];
+
+				bool bShowTangents = TangentVisibility == ECurveEditorTangentVisibility::AllTangents || (TangentVisibility == ECurveEditorTangentVisibility::SelectedKeys && SelectedKeys && SelectedKeys->Contains(VisibleKeys[Index]));
+
+				float TimeScreenPos = CurveSpace.SecondsToScreen(KeyPosition.InputValue + InputOffset);
+				float ValueScreenPos = CurveSpace.ValueToScreen(KeyPosition.OutputValue);
+
+				// Add this key
+				FCurvePointInfo Key(KeyHandle);
+				Key.ScreenPosition = FVector2D(TimeScreenPos, ValueScreenPos);
+				Key.LayerBias = 2;
+
+				// Add draw info for the specific key
+				CurveModel->GetKeyDrawInfo(ECurvePointType::Key, KeyHandle, /*Out*/ Key.DrawInfo);
+				Params.Points.Add(Key);
+
+				if (bShowTangents && Attributes.HasArriveTangent())
+				{
+					float ArriveTangent = Attributes.GetArriveTangent();
+
+					FCurvePointInfo ArriveTangentPoint(KeyHandle);
+					ArriveTangentPoint.Type = ECurvePointType::ArriveTangent;
+
+
+					if (Attributes.HasTangentWeightMode() && Attributes.HasArriveTangentWeight() &&
+						(Attributes.GetTangentWeightMode() == RCTWM_WeightedBoth || Attributes.GetTangentWeightMode() == RCTWM_WeightedArrive))
+					{
+						FVector2D TangentOffset = CurveEditor::ComputeScreenSpaceTangentOffset(CurveSpace, ArriveTangent, -Attributes.GetArriveTangentWeight());
+						ArriveTangentPoint.ScreenPosition = Key.ScreenPosition + TangentOffset;
+					}
+					else
+					{
+						float PixelLength = 60.0f;
+						ArriveTangentPoint.ScreenPosition = Key.ScreenPosition + CurveEditor::GetVectorFromSlopeAndLength(ArriveTangent * -DisplayRatio, -PixelLength);
+					}
+					ArriveTangentPoint.LineDelta = Key.ScreenPosition - ArriveTangentPoint.ScreenPosition;
+					ArriveTangentPoint.LayerBias = 1;
+
+					// Add draw info for the specific tangent
+					FKeyDrawInfo TangentDrawInfo;
+					CurveModel->GetKeyDrawInfo(ECurvePointType::ArriveTangent, KeyHandle, /*Out*/ ArriveTangentPoint.DrawInfo);
+
+					Params.Points.Add(ArriveTangentPoint);
+				}
+
+				if (bShowTangents && Attributes.HasLeaveTangent())
+				{
+					float LeaveTangent = Attributes.GetLeaveTangent();
+
+					FCurvePointInfo LeaveTangentPoint(KeyHandle);
+					LeaveTangentPoint.Type = ECurvePointType::LeaveTangent;
+
+					if (Attributes.HasTangentWeightMode() && Attributes.HasLeaveTangentWeight() &&
+						(Attributes.GetTangentWeightMode() == RCTWM_WeightedBoth || Attributes.GetTangentWeightMode() == RCTWM_WeightedLeave))
+					{
+						FVector2D TangentOffset = CurveEditor::ComputeScreenSpaceTangentOffset(CurveSpace, LeaveTangent, Attributes.GetLeaveTangentWeight());
+
+						LeaveTangentPoint.ScreenPosition = Key.ScreenPosition + TangentOffset;
+					}
+					else
+					{
+						float PixelLength = 60.0f;
+						LeaveTangentPoint.ScreenPosition = Key.ScreenPosition + CurveEditor::GetVectorFromSlopeAndLength(LeaveTangent * -DisplayRatio, PixelLength);
+
+					}
+
+					LeaveTangentPoint.LineDelta = Key.ScreenPosition - LeaveTangentPoint.ScreenPosition;
+					LeaveTangentPoint.LayerBias = 1;
+
+					// Add draw info for the specific tangent
+					FKeyDrawInfo TangentDrawInfo;
+					CurveModel->GetKeyDrawInfo(ECurvePointType::LeaveTangent, KeyHandle, /*Out*/ LeaveTangentPoint.DrawInfo);
+
+					Params.Points.Add(LeaveTangentPoint);
+				}
+			}
+		}
+
+		OutDrawParams.Add(MoveTemp(Params));
+	}
+}
+void SInteractiveCurveEditorView::GetPointsWithinWidgetRange(const FSlateRect& WidgetRectangle, TArray<FCurvePointHandle>* OutPoints) const
+{
+	TSharedPtr<FCurveEditor> CurveEditor = WeakCurveEditor.Pin();
+	if (!CurveEditor)
+	{
+		return;
+	}
+
+	// Iterate through all of our points and see which points the marquee overlaps. Both of these coordinate systems
+	// are in screen space pixels.
+	for (const FCurveDrawParams& DrawParams : CachedDrawParams)
+	{
+		for (int32 PointIndex = 0; PointIndex < DrawParams.Points.Num(); PointIndex++)
+		{
+			const FCurvePointInfo& Point = DrawParams.Points[PointIndex];
+
+			const FKeyDrawInfo& DrawInfo = DrawParams.GetKeyDrawInfo(Point.Type, PointIndex);
+			FSlateRect PointRect = FSlateRect::FromPointAndExtent(Point.ScreenPosition - DrawInfo.ScreenSize/2, DrawInfo.ScreenSize);
+
+			if (FSlateRect::DoRectanglesIntersect(PointRect, WidgetRectangle))
+			{
+				OutPoints->Add(FCurvePointHandle(DrawParams.GetID(), Point.Type, Point.KeyHandle));
+			}
+		}
+	}
+}
+
+void SInteractiveCurveEditorView::UpdateCurveProximities(FVector2D MousePixel)
+{
+	if (DragOperation.IsSet())
+	{
+		return;
+	}
+
+	CurveProximities.Reset();
+	CachedToolTipData.Reset();
+
+	TSharedPtr<FCurveEditor> CurveEditor = WeakCurveEditor.Pin();
+	if (!CurveEditor)
+	{
+		return;
+	}
+
+	TOptional<FCurvePointHandle> MousePoint = HitPoint(MousePixel);
+	if (MousePoint.IsSet())
+	{
+		// If the mouse is over a point, that curve is always the closest, so just add that directly and don't
+		// bother adding the others
+		CurveProximities.Add(MakeTuple(MousePoint->CurveID, 0.f));
+	}
+	else for (const TTuple<FCurveModelID, FCurveInfo>& Pair : CurveInfoByID)
+	{
+		const FCurveModel* CurveModel = CurveEditor->FindCurve(Pair.Key);
+		if (!ensureAlways(CurveModel))
+		{
+			continue;
+		}
+
+		FCurveEditorScreenSpace CurveSpace = GetCurveSpace(Pair.Key);
+
+		double MinMouseTime = CurveSpace.ScreenToSeconds(MousePixel.X - CurveViewConstants::HoverProximityThresholdPx);
+		double MaxMouseTime = CurveSpace.ScreenToSeconds(MousePixel.X + CurveViewConstants::HoverProximityThresholdPx);
+		double MouseValue = CurveSpace.ScreenToValue(MousePixel.Y);
+		float  PixelsPerOutput = CurveSpace.PixelsPerOutput();
+
+		FVector2D MinPos(MousePixel.X - CurveViewConstants::HoverProximityThresholdPx, 0.0f);
+		FVector2D MaxPos(MousePixel.X + CurveViewConstants::HoverProximityThresholdPx, 0.0f);
+
+		double InputOffset = CurveModel->GetInputDisplayOffset();
+		double MinEvalTime = MinMouseTime - InputOffset;
+		double MaxEvalTime = MaxMouseTime - InputOffset;
+
+		double MinValue = 0.0, MaxValue = 0.0;
+		if (CurveModel->Evaluate(MinEvalTime, MinValue) && CurveModel->Evaluate(MaxEvalTime, MaxValue))
+		{
+			MinPos.Y = CurveSpace.ValueToScreen(MinValue);
+			MaxPos.Y = CurveSpace.ValueToScreen(MaxValue);
+
+			float Distance = (FMath::ClosestPointOnSegment2D(MousePixel, MinPos, MaxPos) - MousePixel).Size();
+			CurveProximities.Add(MakeTuple(Pair.Key, Distance));
+		}
+	}
+
+	Algo::SortBy(CurveProximities, [](TTuple<FCurveModelID, float> In) { return In.Get<1>(); });
+
+	if (CurveProximities.Num() > 0 && CurveProximities[0].Get<1>() < CurveViewConstants::HoverProximityThresholdPx)
+	{
+		const FCurveModel* HoveredCurve = CurveEditor->FindCurve(CurveProximities[0].Get<0>());
+		if (HoveredCurve)
+		{
+			FCurveEditorScreenSpace CurveSpace = GetCurveSpace(CurveProximities[0].Get<0>());
+			double MouseTime = CurveSpace.ScreenToSeconds(MousePixel.X) - HoveredCurve->GetInputDisplayOffset();
+			double EvaluatedTime = CurveEditor->GetSnapMetrics().SnapInputSeconds(MouseTime);
+
+			double EvaluatedValue = 0.0;
+			HoveredCurve->Evaluate(EvaluatedTime, EvaluatedValue);
+
+			FCachedToolTipData ToolTipData;
+			ToolTipData.Text = FText::Format(LOCTEXT("CurveEditorTooltipName", "Name: {0}"), HoveredCurve->GetLongDisplayName());
+			ToolTipData.EvaluatedTime = FText::Format(LOCTEXT("CurveEditorTime", "Time: {0}"), EvaluatedTime);
+			ToolTipData.EvaluatedValue = FText::Format(LOCTEXT("CurveEditorValue", "Value: {0}"), EvaluatedValue);
+			
+			CachedToolTipData = ToolTipData;
+		}
+	}
+}
+
+void SInteractiveCurveEditorView::OnMouseEnter(const FGeometry& MyGeometry, const FPointerEvent& MouseEvent)
+{
+	//CurveProximities.Reset();
+	SCompoundWidget::OnMouseEnter(MyGeometry, MouseEvent);
+}
+
+void SInteractiveCurveEditorView::OnMouseLeave(const FPointerEvent& MouseEvent)
+{
+	// Don't allow hover highlights when we've exited this view as clicking won't be routed to us to select it anyways.
+	CurveProximities.Reset();
+	SCompoundWidget::OnMouseLeave(MouseEvent);
+}
+
+FReply SInteractiveCurveEditorView::OnMouseMove(const FGeometry& MyGeometry, const FPointerEvent& MouseEvent)
+{
+	TSharedPtr<FCurveEditor> CurveEditor = WeakCurveEditor.Pin();
+	TSharedPtr<SCurveEditorPanel> EditorPanel = CurveEditor.IsValid() ? CurveEditor->GetPanel() : nullptr;
+	if (!CurveEditor || !EditorPanel)
+	{
+		return FReply::Unhandled();
+	}
+
+	// Don't handle updating if we have a context menu open.
+	if (ActiveContextMenu.Pin())
+	{
+		return FReply::Unhandled();
+	}
+
+	// Update our Curve Proximities for hover states and context actions. This also updates our cached hovered curve.
+	FVector2D MousePixel = MyGeometry.AbsoluteToLocal(MouseEvent.GetScreenSpacePosition());
+	UpdateCurveProximities(MousePixel);
+
+	// Cache the mouse position so that commands such as add key can work from command bindings 
+	CachedMousePosition = MousePixel;
+
+	if (DragOperation.IsSet())
+	{
+		FVector2D InitialPosition = DragOperation->GetInitialPosition();
+
+		if (!DragOperation->IsDragging() && DragOperation->AttemptDragStart(MouseEvent))
+		{
+			DragOperation->DragImpl->BeginDrag(InitialPosition, MousePixel, MouseEvent);
+			return FReply::Handled().CaptureMouse(AsShared());
+		}
+		else if (DragOperation->IsDragging())
+		{
+			DragOperation->DragImpl->Drag(InitialPosition, MousePixel, MouseEvent);
+		}
+		return FReply::Handled();
+	}
+
+	// We don't absorb this event as we're just updating hover states anyways.
+	return FReply::Unhandled();
+}
+
+FReply SInteractiveCurveEditorView::OnMouseWheel(const FGeometry& MyGeometry, const FPointerEvent& MouseEvent)
+{
+	TSharedPtr<FCurveEditor> CurveEditor = WeakCurveEditor.Pin();
+	if (!CurveEditor || bFixedOutputBounds)
+	{
+		return FReply::Unhandled();
+	}
+
+	FCurveEditorScreenSpace ViewSpace = GetViewSpace();
+
+	FVector2D MousePixel   = MyGeometry.AbsoluteToLocal(MouseEvent.GetScreenSpacePosition());
+	double    CurrentTime  = ViewSpace.ScreenToSeconds(MousePixel.X);
+	double    CurrentValue = ViewSpace.ScreenToValue(MousePixel.Y);
+
+	// Attempt to zoom around the current time if settings specify it and there is a valid time.
+	if (CurveEditor->GetSettings()->GetZoomPosition() == ECurveEditorZoomPosition::CurrentTime)
+	{
+		if (CurveEditor->GetTimeSliderController().IsValid())
+		{
+			FFrameTime ScrubPosition = CurveEditor->GetTimeSliderController()->GetScrubPosition();
+			double PlaybackPosition = ScrubPosition / CurveEditor->GetTimeSliderController()->GetTickResolution();
+			if (CurveEditor->GetTimeSliderController()->GetViewRange().Contains(PlaybackPosition))
+			{
+				CurrentTime = PlaybackPosition;
+			}
+		}
+	}
+
+	float ZoomDelta = 1.f - FMath::Clamp(0.1f * MouseEvent.GetWheelDelta(), -0.9f, 0.9f);
+	ZoomAround(FVector2D(ZoomDelta, ZoomDelta), CurrentTime, CurrentValue);
+
+	return FReply::Handled();
+}
+
+TOptional<FCurveModelID> SInteractiveCurveEditorView::GetHoveredCurve() const
+{
+	if (CurveProximities.Num() > 0 && CurveProximities[0].Get<1>() < CurveViewConstants::HoverProximityThresholdPx)
+	{
+		return CurveProximities[0].Get<0>();
+	}
+
+	return TOptional<FCurveModelID>();
+}
+
+bool SInteractiveCurveEditorView::IsToolTipEnabled() const
+{
+	TSharedPtr<FCurveEditor> CurveEditor = WeakCurveEditor.Pin();
+	if (CurveEditor)
+	{
+		return (CachedToolTipData.IsSet() && CurveEditor->GetSettings()->GetShowCurveEditorCurveToolTips());
+	}
+
+	return false;
+}
+
+FText SInteractiveCurveEditorView::GetToolTipCurveName() const
+{
+	return CachedToolTipData.IsSet() ? CachedToolTipData->Text : FText();
+}
+
+FText SInteractiveCurveEditorView::GetToolTipTimeText() const
+{
+	return CachedToolTipData.IsSet() ? CachedToolTipData->EvaluatedTime : FText();
+}
+
+FText SInteractiveCurveEditorView::GetToolTipValueText() const
+{
+	return CachedToolTipData.IsSet() ? CachedToolTipData->EvaluatedValue : FText();
+}
+
+FReply SInteractiveCurveEditorView::OnMouseButtonDown(const FGeometry& MyGeometry, const FPointerEvent& MouseEvent)
+{
+	TSharedPtr<FCurveEditor> CurveEditor = WeakCurveEditor.Pin();
+	TSharedPtr<SCurveEditorPanel> EditorPanel = CurveEditor.IsValid() ? CurveEditor->GetPanel() : nullptr;
+	if (!CurveEditor || !EditorPanel)
+	{
+		return FReply::Unhandled();
+	}
+	
+	FVector2D MousePixel = MyGeometry.AbsoluteToLocal(MouseEvent.GetScreenSpacePosition());
+
+	// Cache the mouse position so that commands such as add key can work from command bindings 
+	CachedMousePosition = MousePixel;
+
+	// Rebind our context actions so that shift click commands use the right position.
+	RebindContextualActions(MousePixel);
+
+	// // Marquee Selection
+	// if (MouseEvent.GetEffectingButton() == EKeys::LeftMouseButton)
+	// {
+	// 	DragOperation = FCurveEditorDelayedDrag(MousePixel, MouseEvent.GetEffectingButton());
+	// 	DragOperation->DragImpl = MakeUnique<FCurveEditorDragOperation_Marquee>(CurveEditor.Get(), this);
+	// 	return FReply::Handled();
+	// }
+	// Middle Click + Alt Pan
+	//else
+	if (MouseEvent.GetEffectingButton() == EKeys::MiddleMouseButton)
+	{
+		if (MouseEvent.IsAltDown())
+		{
+			DragOperation = FCurveEditorDelayedDrag(MousePixel, MouseEvent.GetEffectingButton());
+			DragOperation->DragImpl = MakeUnique<FCurveEditorDragOperation_PanInput>(CurveEditor.Get());
+			return FReply::Handled();
+		}
+		else
+		{
+			// Middle Mouse can try to create keys on curves.
+			TOptional<FCurvePointHandle> NewPoint;
+
+			// Add a key to the closest curve to the mouse
+			if (TOptional<FCurveModelID> HoveredCurve = GetHoveredCurve())
+			{
+				FCurveModel* CurveToAddTo = CurveEditor->FindCurve(HoveredCurve.GetValue());
+				if (CurveToAddTo)
+				{
+					FScopedTransaction Transaction(LOCTEXT("InsertKey", "Insert Key"));
+
+					FCurveEditorScreenSpace CurveSpace = GetCurveSpace(HoveredCurve.GetValue());
+					FKeyAttributes DefaultAttributes = CurveEditor->GetDefaultKeyAttributes().Get();
+
+					double MouseTime = CurveSpace.ScreenToSeconds(MousePixel.X);
+					double MouseValue = CurveSpace.ScreenToValue(MousePixel.Y);
+
+					MouseTime = CurveEditor->GetSnapMetrics().SnapInputSeconds(MouseTime);
+					MouseValue = CurveEditor->GetSnapMetrics().SnapOutput(MouseValue);
+
+					// When adding to a curve with no variance, add it with the same value so that
+					// curves don't pop wildly in normalized views due to a slight difference between the keys
+					double CurveOutputMin = 0, CurveOutputMax = 1;
+					CurveToAddTo->GetValueRange(CurveOutputMin, CurveOutputMax);
+					if (CurveOutputMin == CurveOutputMax)
+					{
+						MouseValue = CurveOutputMin;
+					}
+
+					CurveToAddTo->Modify();
+
+					// Add a key on this curve
+					TOptional<FKeyHandle> NewKey = CurveToAddTo->AddKey(FKeyPosition(MouseTime, MouseValue), DefaultAttributes);
+					if (NewKey.IsSet())
+					{
+						NewPoint = FCurvePointHandle(HoveredCurve.GetValue(), ECurvePointType::Key, NewKey.GetValue());
+
+						CurveEditor->GetSelection().Clear();
+						CurveEditor->GetSelection().Add(NewPoint.GetValue());
+					}
+					else
+					{
+						Transaction.Cancel();
+					}
+				}
+			}
+
+			TUniquePtr<ICurveEditorKeyDragOperation> KeyDrag = CreateKeyDrag(CurveEditor->GetSelection().GetSelectionType());
+
+			KeyDrag->Initialize(CurveEditor.Get(), NewPoint);
+
+			DragOperation = FCurveEditorDelayedDrag(MousePixel, MouseEvent.GetEffectingButton());
+			DragOperation->DragImpl = MoveTemp(KeyDrag);
+
+			return FReply::Handled();
+		}
+	}
+	else if (MouseEvent.GetEffectingButton() == EKeys::RightMouseButton)
+	{
+		// Zoom Timeline
+		if (MouseEvent.IsAltDown())
+		{
+			DragOperation = FCurveEditorDelayedDrag(MousePixel, MouseEvent.GetEffectingButton());
+			DragOperation->DragImpl = MakeUnique<FCurveEditorDragOperation_Zoom>(CurveEditor.Get(), SharedThis(this));
+			return FReply::Handled();
+		}
+		// Pan Timeline if we have flexible output bounds
+		else if (!bFixedOutputBounds)
+		{
+			DragOperation = FCurveEditorDelayedDrag(MousePixel, MouseEvent.GetEffectingButton());
+			DragOperation->DragImpl = MakeUnique<FCurveEditorDragOperation_PanView>(CurveEditor.Get(), SharedThis(this));
+			return FReply::Handled();
+		}
+	}
+
+	bool bShiftPressed = MouseEvent.IsShiftDown();
+	bool bCtrlPressed = MouseEvent.IsControlDown();
+
+	if (MouseEvent.GetEffectingButton() == EKeys::LeftMouseButton)
+	{
+		// Key Selection Testing
+		TOptional<FCurvePointHandle> MouseDownPoint = HitPoint(MousePixel);
+		if (MouseDownPoint.IsSet())
+		{
+			if (bShiftPressed)
+			{
+				CurveEditor->GetSelection().Add(MouseDownPoint.GetValue());
+			}
+			else if (bCtrlPressed)
+			{
+				CurveEditor->GetSelection().Toggle(MouseDownPoint.GetValue());
+			}
+			else
+			{
+				if (CurveEditor->GetSelection().Contains(MouseDownPoint->CurveID, MouseDownPoint->KeyHandle))
+				{
+					CurveEditor->GetSelection().ChangeSelectionPointType(MouseDownPoint->PointType);
+				}
+				else
+				{
+					CurveEditor->GetSelection().Clear();
+					CurveEditor->GetSelection().Add(MouseDownPoint.GetValue());
+				}
+			}
+
+			TUniquePtr<ICurveEditorKeyDragOperation> KeyDrag = CreateKeyDrag(MouseDownPoint->PointType);
+
+			KeyDrag->Initialize(CurveEditor.Get(), MouseDownPoint);
+
+			DragOperation = FCurveEditorDelayedDrag(MousePixel, MouseEvent.GetEffectingButton());
+			DragOperation->DragImpl = MoveTemp(KeyDrag);
+
+			return FReply::Handled();
+		}
+	}
+
+	return FReply::Unhandled();
+}
+
+FReply SInteractiveCurveEditorView::OnMouseButtonUp(const FGeometry& MyGeometry, const FPointerEvent& MouseEvent)
+{
+	TSharedPtr<FCurveEditor> CurveEditor = WeakCurveEditor.Pin();
+	TSharedPtr<SCurveEditorPanel> EditorPanel = CurveEditor.IsValid() ? CurveEditor->GetPanel() : nullptr;
+	if (!CurveEditor || !EditorPanel)
+	{
+		return FReply::Unhandled();
+	}
+
+	FVector2D MousePosition = MyGeometry.AbsoluteToLocal(MouseEvent.GetScreenSpacePosition());
+	if (DragOperation.IsSet())
+	{
+		if (DragOperation->IsDragging())
+		{
+			FVector2D InitialPosition = DragOperation->GetInitialPosition();
+			DragOperation->DragImpl->EndDrag(InitialPosition, MousePosition, MouseEvent);
+
+			DragOperation.Reset();
+			return FReply::Handled().ReleaseMouseCapture();
+		}
+	}
+
+	DragOperation.Reset();
+
+	if (MouseEvent.GetEffectingButton() == EKeys::LeftMouseButton)
+	{
+		// Curve Selection Testing.
+		TOptional<FCurveModelID> HitCurve = GetHoveredCurve();
+		if (!HitPoint(MousePosition).IsSet() && HitCurve.IsSet() && !MouseEvent.IsShiftDown() && !MouseEvent.IsAltDown() && !MouseEvent.IsControlDown())
+		{
+			FCurveModel* CurveModel = CurveEditor->FindCurve(HitCurve.GetValue());
+
+			TArray<FKeyHandle> KeyHandles;
+			KeyHandles.Reserve(CurveModel->GetNumKeys());
+			CurveModel->GetKeys(*CurveEditor, TNumericLimits<double>::Lowest(), TNumericLimits<double>::Max(), TNumericLimits<double>::Lowest(), TNumericLimits<double>::Max(), KeyHandles);
+
+			// Add or remove all keys from the curve.
+			if (MouseEvent.IsShiftDown())
+			{
+				CurveEditor->GetSelection().Add(HitCurve.GetValue(), ECurvePointType::Key, KeyHandles);
+			}
+			else if (MouseEvent.IsControlDown())
+			{
+				CurveEditor->GetSelection().Toggle(HitCurve.GetValue(), ECurvePointType::Key, KeyHandles);
+			}
+			else
+			{
+				CurveEditor->GetSelection().ChangeSelectionPointType(ECurvePointType::Key);
+				CurveEditor->GetSelection().Clear();
+				CurveEditor->GetSelection().Add(HitCurve.GetValue(), ECurvePointType::Key, KeyHandles);
+			}
+
+			return FReply::Handled();
+		}
+	}
+
+	if (MouseEvent.GetEffectingButton() == EKeys::RightMouseButton)
+	{
+		CreateContextMenu(MyGeometry, MouseEvent);
+		return FReply::Handled();
+	}
+
+	// If we hit a curve or another UI element, do not allow mouse input to bubble
+	if (HitPoint(MousePosition) || GetHoveredCurve())
+	{
+		return FReply::Handled();
+	}
+
+	return FReply::Unhandled();
+}
+
+void SInteractiveCurveEditorView::CreateContextMenu(const FGeometry& MyGeometry, const FPointerEvent& MouseEvent)
+{
+	TSharedPtr<FCurveEditor> CurveEditor = WeakCurveEditor.Pin();
+	TSharedPtr<SCurveEditorPanel> EditorPanel = CurveEditor.IsValid() ? CurveEditor->GetPanel() : nullptr;
+	if (!CurveEditor || !EditorPanel)
+	{
+		return;
+	}
+
+	FVector2D MousePosition = MyGeometry.AbsoluteToLocal(MouseEvent.GetScreenSpacePosition());
+
+	TOptional<FCurvePointHandle> MouseUpPoint = HitPoint(MousePosition);
+	
+	// We need to update our curve proximities (again) because OnMouseLeave is called (which clears them) 
+	// before this menu is created due to the parent widget capturing mouse focus. The context menu needs
+	// to know which curve you have highlighted for buffering curves.
+	UpdateCurveProximities(MousePosition);
+
+	// Rebind our context menu actions based on the results of hit-testing
+	RebindContextualActions(MousePosition);
+
+	// if (!MouseUpPoint.IsSet())
+	// {
+	// 	CurveEditor->Selection.Clear();
+	// }
+
+	const bool bCloseAfterSelection = true;
+	FMenuBuilder MenuBuilder(bCloseAfterSelection, EditorPanel->GetCommands());
+
+	FCurveEditorContextMenu::BuildMenu(MenuBuilder, CurveEditor.ToSharedRef(), MouseUpPoint, GetHoveredCurve());
+
+	// Push the context menu
+	FWidgetPath WidgetPath = MouseEvent.GetEventPath() != nullptr ? *MouseEvent.GetEventPath() : FWidgetPath();
+	ActiveContextMenu = FSlateApplication::Get().PushMenu(AsShared(), WidgetPath, MenuBuilder.MakeWidget(), FSlateApplication::Get().GetCursorPos(), FPopupTransitionEffect(FPopupTransitionEffect::ContextMenu));
+}
+
+TOptional<FCurvePointHandle> SInteractiveCurveEditorView::HitPoint(FVector2D MousePixel) const
+{
+	TOptional<FCurvePointHandle> HitPoint;
+	TOptional<float> ClosestDistance;
+
+	TOptional<FCurveModelID> HoveredCurve = GetHoveredCurve();
+
+	// Find all keys within the current hit test time
+	for (const FCurveDrawParams& Params : CachedDrawParams)
+	{
+		// If we have a hovered curve, only hit a point within that curve
+		if (HoveredCurve.IsSet() && Params.GetID() != HoveredCurve.GetValue())
+		{
+			continue;
+		}
+	
+		for (int32 PointIndex = 0; PointIndex < Params.Points.Num(); PointIndex++)
+		{
+			const FCurvePointInfo& Point = Params.Points[PointIndex];
+			const FKeyDrawInfo& PointDrawInfo = Params.GetKeyDrawInfo(Point.Type, PointIndex);
+
+			// We artificially inflate the hit testing region for keys by a few pixels to make them easier to hit. The PointDrawInfo.ScreenSize specifies their drawn size,
+			// so we need to inflate here when doing the actual hit testing. We subtract by half the extent to center it on the drawing.
+			FVector2D HitTestSize = PointDrawInfo.ScreenSize + FVector2D(4.f, 4.f);
+
+			FSlateRect KeyRect = FSlateRect::FromPointAndExtent(Point.ScreenPosition - (HitTestSize / 2.f), HitTestSize);
+
+			if (KeyRect.ContainsPoint(MousePixel))
+			{
+				float DistanceSquared = (KeyRect.GetCenter() - MousePixel).SizeSquared();
+				if (DistanceSquared <= ClosestDistance.Get(DistanceSquared))
+				{
+					ClosestDistance = DistanceSquared;
+					HitPoint = FCurvePointHandle(Params.GetID(), Point.Type, Point.KeyHandle);
+				}
+			}
+		}
+	}
+
+	return HitPoint;
+}
+
+void SInteractiveCurveEditorView::RebindContextualActions(FVector2D InMousePosition)
+{
+	TSharedPtr<FCurveEditor> CurveEditor = WeakCurveEditor.Pin();
+	TSharedPtr<SCurveEditorPanel> CurveEditorPanel = CurveEditor.IsValid() ? CurveEditor->GetPanel() : nullptr;
+	if (!CurveEditorPanel)
+	{
+		return;
+	}
+
+	TSharedPtr<FUICommandList> CommandList = CurveEditorPanel->GetCommands();
+
+	CommandList->UnmapAction(FCurveEditorCommands::Get().AddKeyHovered);
+	CommandList->UnmapAction(FCurveEditorCommands::Get().AddKeyToAllCurves);
+
+	CommandList->UnmapAction(FCurveEditorCommands::Get().BufferVisibleCurves);
+	CommandList->UnmapAction(FCurveEditorCommands::Get().ApplyBufferedCurves);
+	
+
+	TOptional<FCurveModelID> HoveredCurve = GetHoveredCurve();
+	if (HoveredCurve.IsSet())
+	{
+		TSet<FCurveModelID> HoveredCurveSet;
+		HoveredCurveSet.Add(HoveredCurve.GetValue());
+
+		CommandList->MapAction(FCurveEditorCommands::Get().AddKeyHovered, FExecuteAction::CreateSP(this, &SInteractiveCurveEditorView::AddKeyAtMousePosition, HoveredCurveSet));
+
+		// Buffer the curve they have highlighted instead of all of them.
+		CommandList->MapAction(FCurveEditorCommands::Get().BufferVisibleCurves, FExecuteAction::CreateSP(this, &SInteractiveCurveEditorView::BufferCurve, HoveredCurve.GetValue()));
+	}
+	else
+	{
+		// Apply the buffering action to our entire set and not just the hovered curve.
+		CommandList->MapAction(FCurveEditorCommands::Get().BufferVisibleCurves, FExecuteAction::CreateSP(this, &SInteractiveCurveEditorView::BufferVisibleCurves));
+	}
+
+	CommandList->MapAction(FCurveEditorCommands::Get().AddKeyToAllCurves, FExecuteAction::CreateSP(this, &SInteractiveCurveEditorView::AddKeyAtScrubTime, TSet<FCurveModelID>()));
+
+	// Buffer Visible Curves. Can only apply buffered curves if the current number of visible curves matches the number of buffered curves.
+	CommandList->MapAction(FCurveEditorCommands::Get().ApplyBufferedCurves, FExecuteAction::CreateSP(this, &SInteractiveCurveEditorView::ApplyBufferCurves, HoveredCurve), FCanExecuteAction::CreateSP(this, &SInteractiveCurveEditorView::CanApplyBufferedCurves, HoveredCurve));
+}
+
+void SInteractiveCurveEditorView::BufferVisibleCurves()
+{
+	TSharedPtr<FCurveEditor> CurveEditor = WeakCurveEditor.Pin();
+	if (CurveEditor.IsValid())
+	{
+		// Curve Editor will handle copying and storing the curves.
+		TSet<FCurveModelID> ActiveCurveIDs;
+		for (const TTuple<FCurveModelID, FCurveInfo>& Pair : CurveInfoByID)
+		{
+			ActiveCurveIDs.Add(Pair.Key);
+		}
+		CurveEditor->SetBufferedCurves(ActiveCurveIDs);
+	}
+}
+
+void SInteractiveCurveEditorView::BufferCurve(const FCurveModelID CurveID)
+{
+	TSharedPtr<FCurveEditor> CurveEditor = WeakCurveEditor.Pin();
+	if (CurveEditor.IsValid())
+	{
+		// Curve Editor will handle copying and storing the curves.
+		TSet<FCurveModelID> CurveSet;
+		CurveSet.Add(CurveID);
+		CurveEditor->SetBufferedCurves(CurveSet);
+	}
+}
+
+void SInteractiveCurveEditorView::ApplyBufferCurves(TOptional<FCurveModelID> DestinationCurve)
+{
+	TSharedPtr<FCurveEditor> CurveEditor = WeakCurveEditor.Pin();
+	if (CurveEditor.IsValid())
+	{
+		if (DestinationCurve.IsSet())
+		{
+			TSet<FCurveModelID> CurveSet;
+			CurveSet.Add(DestinationCurve.GetValue());
+
+			// Apply the buffered curve (singular) to our highlighted curve.
+			CurveEditor->ApplyBufferedCurves(CurveSet);
+		}
+		else
+		{
+			// Curve Editor will handle attempting to apply the buffered curves to our currently visible ones.
+			TSet<FCurveModelID> ActiveCurveIDs;
+			for (const TTuple<FCurveModelID, FCurveInfo>& Pair : CurveInfoByID)
+			{
+				ActiveCurveIDs.Add(Pair.Key);
+			}
+			CurveEditor->ApplyBufferedCurves(ActiveCurveIDs);
+		}
+
+	}
+}
+
+bool SInteractiveCurveEditorView::CanApplyBufferedCurves(TOptional<FCurveModelID> DestinationCurve) const
+{
+	TSharedPtr<FCurveEditor> CurveEditor = WeakCurveEditor.Pin();
+	if (CurveEditor.IsValid())
+	{
+		if (DestinationCurve.IsSet())
+		{
+			return CurveEditor->GetNumBufferedCurves() == 1;
+		}
+		else
+		{
+			// For now we just do a 1:1 mapping. Once curves have better names we can try to do an intelligent match up, ie: matching Transform.X to a new Transform.X
+			return CurveEditor->GetNumBufferedCurves() == NumCurves();
+		}
+	}
+
+	return false;
+}
+
+void SInteractiveCurveEditorView::AddKeyAtScrubTime(TSet<FCurveModelID> ForCurves)
+{
+	TSharedPtr<FCurveEditor> CurveEditor = WeakCurveEditor.Pin();
+	if (!CurveEditor.IsValid())
+	{
+		return;
+	}
+
+	TSet<FCurveModelID> CurvesToAddTo;
+	if (ForCurves.Num() == 0)
+	{
+		CurvesToAddTo = CurveEditor->GetEditedCurves();
+	}
+	else
+	{
+		CurvesToAddTo = ForCurves;
+	}
+
+	// If they don't have a time slider controller then we fall back to using mouse position.
+	TSharedPtr<ITimeSliderController> TimeSliderController = CurveEditor->GetTimeSliderController();
+	if (!TimeSliderController)
+	{
+		AddKeyAtMousePosition(CurvesToAddTo);
+		return;
+	}
+
+	// Snapping of the time will be done inside AddKeyAtTime.
+	double ScrubTime = TimeSliderController->GetScrubPosition() / TimeSliderController->GetTickResolution();
+	AddKeyAtTime(CurvesToAddTo, ScrubTime);
+}
+
+void SInteractiveCurveEditorView::AddKeyAtMousePosition(TSet<FCurveModelID> ForCurves)
+{
+	TSharedPtr<FCurveEditor> CurveEditor = WeakCurveEditor.Pin();
+	if (!CurveEditor.IsValid())
+	{
+		return;
+	}
+
+	// Snapping will be done inside AddKeyAtTime
+	double MouseTime = GetViewSpace().ScreenToSeconds(CachedMousePosition.X);
+	AddKeyAtTime(ForCurves, MouseTime);
+}
+
+void SInteractiveCurveEditorView::AddKeyAtTime(const TSet<FCurveModelID>& ToCurves, double InTime)
+{
+	TSharedPtr<FCurveEditor> CurveEditor = WeakCurveEditor.Pin();
+	if (!CurveEditor.IsValid())
+	{
+		return;
+	}
+
+	// Ensure the time is snapped if needed
+	InTime = CurveEditor->GetSnapMetrics().SnapInputSeconds(InTime);
+
+	FScopedTransaction Transaction(LOCTEXT("AddKeyAtTime", "Add Key"));
+	bool bAddedKey = false;
+
+	FKeyAttributes DefaultAttributes = CurveEditor->GetDefaultKeyAttribute().Get();
+
+	// Clear the selection set as we will be selecting all the new keys created.
+	CurveEditor->GetSelection().Clear();
+
+	for (const FCurveModelID CurveModelID : ToCurves)
+	{
+		FCurveModel* CurveModel = CurveEditor->FindCurve(CurveModelID);
+		check(CurveModel);
+
+		// Support optional input display offsets 
+		double EvalTime = InTime - CurveModel->GetInputDisplayOffset();
+
+		double CurveValue = 0.0;
+		if (CurveModel->Evaluate(EvalTime, CurveValue))
+		{
+			CurveModel->Modify();
+			CurveValue = CurveEditor->GetSnapMetrics().SnapOutput(CurveValue);
+
+			// Curve Models allow us to create new keys ontop of existing keys which works, but causes some user confusion
+			// Before we create a key, we instead check to see if there is already a key at this time, and if there is, we
+			// add that key to the selection set instead. This solves issues with snapping causing keys to be created adjacent
+			// to the mouse cursor (sometimes by a large amount).
+			TArray<FKeyHandle> ExistingKeys;
+			CurveModel->GetKeys(*CurveEditor, EvalTime - KINDA_SMALL_NUMBER, EvalTime + KINDA_SMALL_NUMBER, TNumericLimits<double>::Lowest(), TNumericLimits<double>::Max(), ExistingKeys);
+			
+			TOptional<FKeyHandle> NewKey;
+
+			if (ExistingKeys.Num() > 0)
+			{
+				NewKey = ExistingKeys[0];
+			}
+			else
+			{
+				// Add a key on this curve
+				NewKey = CurveModel->AddKey(FKeyPosition(EvalTime, CurveValue), DefaultAttributes);
+			}
+
+			// Add the key to the selection set.
+			if (NewKey.IsSet())
+			{
+				bAddedKey = true;
+				CurveEditor->GetSelection().Add(FCurvePointHandle(CurveModelID, ECurvePointType::Key, NewKey.GetValue()));
+			}
+		}
+	}
+
+	if (!bAddedKey)
+	{
+		Transaction.Cancel();
+	}
+}
+
+void SInteractiveCurveEditorView::OnCurveEditorToolChanged(FCurveEditorToolID InToolId)
+{
+	// We need to end drag-drop operations if they switch tools. Otherwise they can start
+	// a marquee select, use the keyboard to switch to a diferent tool, and then the marquee
+	// select finishes after the tool has had a chance to activate.
+	if (DragOperation.IsSet())
+	{
+		// We have to cancel it instead of ending it because ending it needs mouse position and some other stuff.
+		DragOperation->DragImpl->CancelDrag();
+		DragOperation.Reset();
+	}
+}
+
+#undef LOCTEXT_NAMESPACE // "SInteractiveCurveEditorView"
diff --git a/Engine/Source/Editor/CurveEditor/Public/CurveDataAbstraction.h b/Engine/Source/Editor/CurveEditor/Public/CurveDataAbstraction.h
--- a/Engine/Source/Editor/CurveEditor/Public/CurveDataAbstraction.h
+++ b/Engine/Source/Editor/CurveEditor/Public/CurveDataAbstraction.h
@@ -173,6 +173,8 @@ struct FCurveAttributes
 	{
 		bHasPreExtrapolation  = 0;
 		bHasPostExtrapolation = 0;
+		PreExtrapolation = ERichCurveExtrapolation::RCCE_None;
+		PostExtrapolation = ERichCurveExtrapolation::RCCE_None;
 	}
 
 	/**
diff --git a/Engine/Source/Editor/CurveEditor/Public/CurveDrawInfo.h b/Engine/Source/Editor/CurveEditor/Public/CurveDrawInfo.h
--- a/Engine/Source/Editor/CurveEditor/Public/CurveDrawInfo.h
+++ b/Engine/Source/Editor/CurveEditor/Public/CurveDrawInfo.h
@@ -56,6 +56,8 @@ struct FCurvePointInfo
 	ECurvePointType Type;
 	/** A layer bias to draw the point with (higher integers draw on top) */
 	int32 LayerBias;
+	/** DrawInfo to use when rendering this point. Allows customization of points on a per-point basis instead of per-curve. */
+	FKeyDrawInfo DrawInfo;
 };
 
 
@@ -73,9 +75,6 @@ struct FCurveDrawParams
 	/** An array of distinct curve points for the visible range. */
 	TArray<FCurvePointInfo> Points;
 
-	/** Value defining how to draw keys of type ECurvePointType::Key. */
-	FKeyDrawInfo KeyDrawInfo;
-
 	/** Value defining how to draw keys of type ECurvePointType::ArriveTangent. */
 	FKeyDrawInfo ArriveTangentDrawInfo;
 
@@ -101,18 +100,22 @@ struct FCurveDrawParams
 	/**
 	 * Retrieve the draw information for drawing the specified type of curve point
 	 */
-	const FKeyDrawInfo& GetKeyDrawInfo(ECurvePointType Type) const
+	const FKeyDrawInfo& GetKeyDrawInfo(ECurvePointType Type, const int32 InPointIndex) const
 	{
 		switch (Type)
 		{
 		case ECurvePointType::ArriveTangent: return ArriveTangentDrawInfo;
 		case ECurvePointType::LeaveTangent:  return LeaveTangentDrawInfo;
-		default:                             return KeyDrawInfo;
+		case ECurvePointType::Key:			 return Points[InPointIndex].DrawInfo;
 		}
+		return InvalidDrawInfo;
 	}
 
 private:
 
 	/** Immutable curve ID */
 	FCurveModelID ID;
+
+	/** Immutable dummy key draw info to return */
+	FKeyDrawInfo InvalidDrawInfo;
 };
\ No newline at end of file
diff --git a/Engine/Source/Editor/CurveEditor/Public/CurveEditor.h b/Engine/Source/Editor/CurveEditor/Public/CurveEditor.h
--- a/Engine/Source/Editor/CurveEditor/Public/CurveEditor.h
+++ b/Engine/Source/Editor/CurveEditor/Public/CurveEditor.h
@@ -13,19 +13,34 @@
 #include "CurveEditorTypes.h"
 #include "CurveDataAbstraction.h"
 #include "CurveModel.h"
+#include "Tree/CurveEditorTree.h"
 #include "ICurveEditorBounds.h"
 #include "CurveEditorSelection.h"
 #include "ICurveEditorDragOperation.h"
+#include "ICurveEditorModule.h"
+#include "ICurveEditorToolExtension.h"
+#include "EditorUndoClient.h"
+#include "CurveEditorScreenSpace.h"
+#include "ITimeSlider.h"
 
 class UCurveEditorSettings;
 class FUICommandList;
+class FCurveEditor;
+class SCurveEditorPanel;
+
 struct FGeometry;
 struct FCurveEditorSnapMetrics;
-struct FCurveEditorScreenSpace;
+class ICurveEditorExtension;
+class ICurveEditorToolExtension;
 
 DECLARE_DELEGATE_OneParam(FOnSetBoolean, bool)
+DECLARE_MULTICAST_DELEGATE_OneParam(FOnCurveArrayChanged, FCurveModelID)
+DECLARE_MULTICAST_DELEGATE_OneParam(FOnActiveToolChanged, FCurveEditorToolID)
+
 
-class CURVEEDITOR_API FCurveEditor : public TSharedFromThis<FCurveEditor>
+class CURVEEDITOR_API FCurveEditor 
+	: public FEditorUndoClient
+	, public TSharedFromThis<FCurveEditor>
 {
 public:
 
@@ -35,6 +50,16 @@ public:
 	FCurveEditorSelection Selection;
 
 	
+protected:
+	struct FBufferedCurve
+	{
+		/** Cached Positions from original curve. */
+		TArray<FKeyPosition>	KeyPositions;
+		/** Cached Attributes from original curve. */
+		TArray<FKeyAttributes>	KeyAttributes;
+		/** The original intention of the type of curve. Can be empty. Used to try and match when pasting multiple curves. */
+		FString					IntentionName;
+	};
 
 public:
 
@@ -59,6 +84,12 @@ public:
 	/** Attribute used for determining default attributes to apply to a newly create key */
 	TAttribute<FKeyAttributes> DefaultKeyAttributes;
 
+	/** Grid line label text format strings for the X and Y axis */
+	TAttribute<FText> GridLineLabelFormatXAttribute, GridLineLabelFormatYAttribute;
+
+
+	/** Delegate that is invoked when a tool becomes active. Also fired when the tool goes inactive. */
+	FOnActiveToolChanged OnActiveToolChangedDelegate;
 public:
 
 	/**
@@ -74,19 +105,85 @@ public:
 
 	virtual ~FCurveEditor() {}
 
+	void InitCurveEditor(const FCurveEditorInitParams& InInitParams);
+
 public:
 
+	void SetPanel(TSharedPtr<SCurveEditorPanel> InPanel);
+
+	TSharedPtr<SCurveEditorPanel> GetPanel() const;
+
+	FCurveEditorScreenSpaceH GetPanelInputSpace() const;
+
+public:
 	/**
-	 * Generate a utility struct for converting between screen (slate unit) space and the underlying input/output axes
+	 * Zoom the curve editor to fit all the currently visible curves
+	 *
+	 * @param Axes         (Optional) Axes to lock the zoom to
 	 */
-	FCurveEditorScreenSpace GetScreenSpace() const;
-
+	void ZoomToFit(EAxisList::Type Axes = EAxisList::All);
 	/**
-	 * Generate a utility struct for snapping values
+	 * Zoom the curve editor to fit the requested curves.
+	 *
+	 * @param CurveModelIDs The curve IDs to zoom to fit.
+	 * @param Axes         (Optional) Axes to lock the zoom to
 	 */
-	FCurveEditorSnapMetrics GetSnapMetrics() const;
+	void ZoomToFitCurves(TArrayView<const FCurveModelID> CurveModelIDs, EAxisList::Type Axes = EAxisList::All);
+	/**
+	 * Zoom the curve editor to fit all the current key selection. Zooms to fit all if less than 2 keys are selected.
+	 *
+	 * @param Axes         (Optional) Axes to lock the zoom to
+	 */
+	void ZoomToFitSelection(EAxisList::Type Axes = EAxisList::All);
+	/**
+	* Assign a new bounds container to this curve editor
+	*/
+	void SetBounds(TUniquePtr<ICurveEditorBounds>&& InBounds);
+	/**
+	 * Retrieve the current curve editor bounds implementation
+	 */
+	const ICurveEditorBounds& GetBounds() const { return *Bounds.Get(); }
+	/**
+	 * Retrieve the current curve editor bounds implementation
+	 */
+	ICurveEditorBounds& GetBounds() { return *Bounds.Get();	}
+	/*
+	 * Sets a Time Slider controller for this Curve Editor to be sync'd against. Can be null.
+	 */
+	void SetTimeSliderController(TSharedPtr<ITimeSliderController> InTimeSliderController) { WeakTimeSliderController = InTimeSliderController; }
+	/**
+	 * Retrieve the optional Time Slider Controller that this Curve Editor may be sync'd with.
+	 */
+	TSharedPtr<ITimeSliderController> GetTimeSliderController() const { return WeakTimeSliderController.Pin(); }
+	/**
+	* Retrieve this curve editor's command list
+	*/
+	TSharedPtr<FUICommandList> GetCommands() const { return CommandList; }
+	/**
+	* Returns true of the specified tool is currently active.
+	*/
+	bool IsToolActive(const FCurveEditorToolID InToolID) const;
+	/**
+	* Attempts to make the specified tool the active tool. This will cancel the current tool if there is one.
+	*/
+	void MakeToolActive(const FCurveEditorToolID InToolID);
+	/**
+	* Attempts to get the currently active tool. Will return nullptr if there is no active tool.
+	* Do not store a reference to this returned pointer, instead only store FCurveEditorToolIDs!
+	*/
+	ICurveEditorToolExtension* GetCurrentTool() const;
 
-public:
+	FCurveEditorToolID AddTool(TUniquePtr<ICurveEditorToolExtension>&& InTool);
+
+	/** Step to next or previous key from the current time */
+	void StepToNextKey();
+	void StepToPreviousKey();
+
+	/** Step forward/backward, jump to start/end */
+	void StepForward();
+	void StepBackward();
+	void JumpToStart();
+	void JumpToEnd();
 
 	/**
 	 * Find a curve by its ID
@@ -94,100 +191,142 @@ public:
 	 * @return a ptr to the curve if found, nullptr otherwise
 	 */
 	FCurveModel* FindCurve(FCurveModelID CurveID) const;
-
 	/**
-	 * Access all the curves currently being shown on this editor
-	 */
-	const TMap<FCurveModelID, TUniquePtr<FCurveModel>>& GetCurves() const;
-
-	/**
-	 * Add a new curve to this editor
-	 */
+	* Add a new curve to this editor
+	*/
 	FCurveModelID AddCurve(TUniquePtr<FCurveModel>&& InCurve);
+	FCurveModelID AddCurveForTreeItem(TUniquePtr<FCurveModel>&& InCurve, FCurveEditorTreeItemID TreeItemID);
 
 	/**
-	 * Remove a curve from this editor
-	 */
+	* Remove a curve from this editor.
+	*/
 	void RemoveCurve(FCurveModelID InCurveID);
 
+	bool IsCurvePinned(FCurveModelID InCurveID) const;
+
+	void PinCurve(FCurveModelID InCurveID);
+
+	void UnpinCurve(FCurveModelID InCurveID);
+
+	const TSet<FCurveModelID>& GetPinnedCurves() const
+	{
+		return PinnedCurves;
+	}
+
+	const SCurveEditorView* FindFirstInteractiveView(FCurveModelID InCurveID) const;
+
 	/**
-	 * Retrieve this curve editor's command list
+	 * Retrieve this curve editor's settings
 	 */
-	FORCEINLINE TSharedPtr<FUICommandList> GetCommands() const { return CommandList; }
+	UCurveEditorSettings* GetSettings() const { return Settings; }
+	/**
+	* Access all the curves currently contained in the Curve Editor regardless of visibility.
+	*/
+	const TMap<FCurveModelID, TUniquePtr<FCurveModel>>& GetCurves() const;
 
+	FCurveEditorSelection& GetSelection() { return Selection; }
 	/**
-	 * Retrieve this curve editor's settings
+	 * Generate a utility struct for snapping values
 	 */
-	FORCEINLINE UCurveEditorSettings* GetSettings() const { return Settings; }
+	FCurveEditorSnapMetrics GetSnapMetrics() const;
+	/**
+	 * Returned the cached struct for snapping editing movement to a specific axis based on user preferences.
+	 */
+	FCurveEditorAxisSnap GetAxisSnap() const { return AxisSnapMetrics; }
+	void SetAxisSnap(const FCurveEditorAxisSnap& InAxisSnap) { AxisSnapMetrics = InAxisSnap; }
+	TAttribute<FText> GetGridLineLabelFormatXAttribute() const { return GridLineLabelFormatXAttribute; }
+	TAttribute<FText> GetGridLineLabelFormatYAttribute() const { return GridLineLabelFormatYAttribute; }
+	TAttribute<FKeyAttributes> GetDefaultKeyAttributes() const { return DefaultKeyAttributes; }
+
+	void SuppressBoundTransformUpdates(bool bSuppress) { bBoundTransformUpdatesSuppressed = bSuppress; }
+	bool AreBoundTransformUpdatesSuppressed() const { return bBoundTransformUpdatesSuppressed; }
+
+	TSet<FCurveModelID> GetEditedCurves() const;
+	/** Attribute used for determining default attributes to apply to a newly create key */
+	TAttribute<FKeyAttributes> GetDefaultKeyAttribute() const { return DefaultKeyAttributes; }
+	/** Create a copy of the specified set of curves. */
+	void SetBufferedCurves(const TSet<FCurveModelID>& InCurves);
+	/** Attempts to apply the buffered curve to the passed in curve set. Returns true on success. */
+	bool ApplyBufferedCurves(const TSet<FCurveModelID>& InCurvesToApplyTo);
+	/** Return the number of stored Buffered Curves. */
+	int32 GetNumBufferedCurves() const { return BufferedCurves.Num(); }
+	// ~FCurveEditor
+
+	// FEditorUndoClient
+	virtual void PostUndo(bool bSuccess);
+	// ~FEditorUndoClient
+
+	const TArray<TSharedRef<ICurveEditorExtension>> GetEditorExtensions() const
+	{
+		return EditorExtensions;
+	}
+
+	const TMap<FCurveEditorToolID, TUniquePtr<ICurveEditorToolExtension>>& GetToolExtensions() const
+	{
+		return ToolExtensions;
+	}
 
 public:
 
 	/**
-	 * Assign a new bounds container to this curve editor
+	 * Retrieve a tree item from its ID
 	 */
-	void SetBounds(TUniquePtr<ICurveEditorBounds>&& InBounds);
+	FCurveEditorTreeItem& GetTreeItem(FCurveEditorTreeItemID ItemID);
 
 	/**
-	 * Retrieve the current curve editor bounds implementation
+	 * Retrieve a tree item from its ID
 	 */
-	const ICurveEditorBounds& GetBounds() const
-	{
-		return *Bounds.Get();
-	}
+	const FCurveEditorTreeItem& GetTreeItem(FCurveEditorTreeItemID ItemID) const;
 
 	/**
-	 * Retrieve the current curve editor bounds implementation
+	 * Get const access to the entire set of root tree items
 	 */
-	ICurveEditorBounds& GetBounds()
-	{
-		return *Bounds.Get();
-	}
+	const TArray<FCurveEditorTreeItemID>& GetRootTreeItems() const;
 
-public:
+	/**
+	 * Add a new tree item to this curve editor
+	 */
+	FCurveEditorTreeItem* AddTreeItem(FCurveEditorTreeItemID ParentID);
 
 	/**
-	 * Check whether this curve editor can automatically zoom to the current selection
+	 * Remove a tree item from the curve editor
 	 */
-	bool ShouldAutoFrame() const;
+	void RemoveTreeItem(FCurveEditorTreeItemID ItemID);
 
 	/**
-	 * Zoom the curve editor in or out around the center point
-	 *
-	 * @param Amount       The amount to zoom by as a factor of the current size
+	 * Retrieve an event that is broadcast any time the tree changes
 	 */
-	void Zoom(float Amount);
+	FSimpleMulticastDelegate& OnTreeChanged();
 
 	/**
-	 * Zoom the curve editor in or out around the specified point
-	 *
-	 * @param Amount       The amount to zoom by as a factor of the current size
-	 * @param TimeOrigin   The time origin to zoom around
-	 * @param ValueOrigin  The value origin to zoom around
+	 * Set the current tree selection
 	 */
-	void ZoomAround(float Amount, double TimeOrigin, double ValueOrigin);
+	void SetDirectTreeSelection(TArray<FCurveEditorTreeItemID>&& TreeItems);
 
 	/**
-	 * Zoom the curve editor to fit all the currently visible curves
-	 *
-	 * @param Axes         (Optional) Axes to lock the zoom to
+	 * Check whether this tree item is selected
 	 */
-	void ZoomToFit(EAxisList::Type Axes = EAxisList::All);
+	ECurveEditorTreeSelectionState GetTreeSelectionState(FCurveEditorTreeItemID TreeItemID) const;
 
 	/**
-	 * Zoom the curve editor to fit the requested curves.
-	 *
-	 * @param CurveModelIDs The curve IDs to zoom to fit.
-	 * @param Axes         (Optional) Axes to lock the zoom to
+	 * Retrieve the current tree selection
 	 */
-	void ZoomToFitCurves(TArrayView<const FCurveModelID> CurveModelIDs, EAxisList::Type Axes = EAxisList::All);
+	const TMap<FCurveEditorTreeItemID, ECurveEditorTreeSelectionState>& GetTreeSelection() const;
 
 	/**
-	 * Zoom the curve editor to fit all the current key selection. Zooms to fit all if less than 2 keys are selected.
-	 *
-	 * @param Axes         (Optional) Axes to lock the zoom to
+	 * Retrieve a serial number that is incremented any time a curve is added or removed
 	 */
-	void ZoomToFitSelection(EAxisList::Type Axes = EAxisList::All);
+	uint32 GetActiveCurvesSerialNumber() const
+	{
+		return ActiveCurvesSerialNumber;
+	}
 
+public:
+
+	/**
+	 * Check whether this curve editor can automatically zoom to the current selection
+	 */
+	bool ShouldAutoFrame() const;
 public:
 
 	/**
@@ -220,25 +359,11 @@ public:
 	void StraightenSelection();
 
 	/**
-	 * Bake curves between selected keys by adding points at every frame of the display rate
-	 */
-	void BakeSelection();
-
-	/**
-	 * Simplify curves between the selected keys by removing redundant keys
-	 * @param Tolerance   Threshold at which to remove keys
-	 */
-	void SimplifySelection(float Tolerance= 0.1);
+	* Do we currently have keys to flatten or straighten?
+	*/
+	bool CanFlattenOrStraightenSelection() const;
 
 public:
-
-	/**
-	 * Populate the specified array with curve painting parameters
-	 *
-	 * @param OutDrawParams    An array to populate with curve painting parameters, one per visible curve
-	 */
-	void GetCurveDrawParams(TArray<FCurveDrawParams>& OutDrawParams) const;
-
 	/**
 	 * Called by SCurveEditorPanel to update the allocated geometry for this curve editor.
 	 */
@@ -247,43 +372,21 @@ public:
 public:
 
 	/**
-	 * Called by SCurveEditorPanel to determine where to draw grid lines along the X-axis
+	 * Called by SCurveEditorPanel to determine where to draw grid lines along the X-axis. This allows
+	 * synchronization with an external data source (such as Sequencer's Timeline ticker). A similar
+	 * function for the Y grid lines is not provided due to the Curve Editor's ability to have multiple
+	 * views with repeated gridlines and values.
 	 */
 	virtual void GetGridLinesX(TArray<float>& MajorGridLines, TArray<float>& MinorGridLines, TArray<FText>& MajorGridLabels) const
 	{
 		ConstructXGridLines(MajorGridLines, MinorGridLines, MajorGridLabels);
 	}
 
-	/**
-	 * Called by SCurveEditorPanel to determine where to draw grid lines along the Y-axis
-	 */
-	virtual void GetGridLinesY(TArray<float>& MajorGridLines, TArray<float>& MinorGridLines, TArray<FText>& MajorGridLabels) const
-	{
-		ConstructYGridLines(MajorGridLines, MinorGridLines, MajorGridLabels, 4);
-	}
-
 	/**
 	 * Bind UI commands that this curve editor responds to
 	 */
 	void BindCommands();
 
-public:
-	/**
-	* Get A Vector for the given slope, usually a tangent, and length. Used to draw the tangent.
-	*/
-	static FVector2D GetVectorFromSlopeAndLength(float Slope, float Length);
-
-public:
-	/**
-	* Given the position of a tangent in screen space get it's position in normal time/value space.
-	*/
-	FVector2D GetTangentPositionInScreenSpace(const FVector2D &StartPos, float Tangent, float Weight) const;
-
-	/**
-	* Given point and tangent position in screen space, get the tangent and it's weight value in normal time/value space.
-	*/
-	void GetTangentAndWeightFromScreenPosition(const FVector2D &StartPos, const  FVector2D &TangentPos, float &Tangent, float &Weight) const;
-
 protected:
 
 	/**
@@ -292,22 +395,39 @@ protected:
 	void ConstructXGridLines(TArray<float>& MajorGridLines, TArray<float>& MinorGridLines, TArray<FText>& MajorGridLabels) const;
 
 	/**
-	 * Construct grid lines for the current visible value range
+	 * Internal zoom to fit implementation
+	 *
+	 * @param Axes        The axes to zoom (only X or Y supported)
+	 * @param CurveKeySet Map from curve ID to keys that should be considered for zoom. Empty key sets will cause the entire curve range to be zoomed.
 	 */
-	void ConstructYGridLines(TArray<float>& MajorGridLines, TArray<float>& MinorGridLines, TArray<FText>& MajorGridLabels, uint8 MinorDivisions) const;
+	void ZoomToFitInternal(EAxisList::Type Axes, const TMap<FCurveModelID, FKeyHandleSet>& CurveKeySet);
 
 	/**
-	 * Internal zoom to fit implementation
-	 */
-	void ZoomToFitInternal(EAxisList::Type Axes, double InputMin, double InputMax, double OutputMin, double OutputMax);
+	*	Apply a specific buffered curve to a specific target curve.
+	*/
+	void ApplyBufferedCurveToTarget(const FBufferedCurve& BufferedCurve, FCurveModel* TargetCurve);
+
 
 protected:
 
 	/** Curve editor bounds implementation */
 	TUniquePtr<ICurveEditorBounds> Bounds;
 
-	/** Map of all the currently visible curve models */
+	/** Map from curve model ID to the actual curve model */
 	TMap<FCurveModelID, TUniquePtr<FCurveModel>> CurveData;
+	/** Map from curve model ID to its originating tree item */
+	TMap<FCurveModelID, FCurveEditorTreeItemID> TreeIDByCurveID;
+
+	/** Set of pinned curve models */
+	TSet<FCurveModelID> PinnedCurves;
+
+	TWeakPtr<SCurveEditorPanel> WeakPanel;
+
+	/** Hierarchical information pertaining to curve data */
+	FCurveEditorTree Tree;
+
+	/** The currently active tool if any. If unset then no tool is currently active and the next selection will default to the first tool. */
+	TOptional<FCurveEditorToolID> ActiveTool;
 
 	/** UI command list of actions mapped to this curve editor */
 	TSharedPtr<FUICommandList> CommandList;
@@ -315,8 +435,35 @@ protected:
 	/** Curve editor settings object */
 	UCurveEditorSettings* Settings;
 
+	/** List of editor extensions we have initialized. */
+	TArray<TSharedRef<ICurveEditorExtension>> EditorExtensions;
+
+	/** List of tool extensions we have initialized. */
+	TMap<FCurveEditorToolID, TUniquePtr<ICurveEditorToolExtension>> ToolExtensions;
+
+	/** Optional external Time Slider controller to sync with. Enables some additional functionality. */
+	TWeakPtr<ITimeSliderController> WeakTimeSliderController;
+
+	/** 
+	* Should attempts to update the bounds of each curve be ignored? This allows tools to keep the bounds from being automatically updated each frame 
+	* which allows Normalized views to push past their boundaries without the normalization ratio changing per-frame as you drag.
+	*/
+	bool bBoundTransformUpdatesSuppressed;
+	
+	/** Track which axis UI movements should be snapped to (where applicable) based on limitations imposed by the UI. */
+	FCurveEditorAxisSnap AxisSnapMetrics;
+
+
+
+	/** Buffered Curves. When a curve is buffered it is copied and the new copy is uniquely owned by the Curve Editor. */
+	TArray<FBufferedCurve> BufferedCurves;
+
+	/** A serial number that is incremented any time the currently active set of curves are changed */
+	uint32 ActiveCurvesSerialNumber;
+
 private:
 
 	/** Cached physical size of the panel representing this editor */
 	FVector2D CachedPhysicalSize;
+
 };
diff --git a/Engine/Source/Editor/CurveEditor/Public/CurveEditorCommands.h b/Engine/Source/Editor/CurveEditor/Public/CurveEditorCommands.h
--- a/Engine/Source/Editor/CurveEditor/Public/CurveEditorCommands.h
+++ b/Engine/Source/Editor/CurveEditor/Public/CurveEditorCommands.h
@@ -66,7 +66,32 @@ public:
 	TSharedPtr<FUICommandInfo> AddKeyHovered;
 
 	TSharedPtr<FUICommandInfo> AddKeyToAllCurves;
-	TSharedPtr<FUICommandInfo> AddKeyToAllCurvesHere;
+
+	TSharedPtr<FUICommandInfo> SetViewModeAbsolute;
+	TSharedPtr<FUICommandInfo> SetViewModeStacked;
+	TSharedPtr<FUICommandInfo> SetViewModeNormalized;
+
+	TSharedPtr<FUICommandInfo> DeactivateCurrentTool;
+	TSharedPtr<FUICommandInfo> DeselectAllKeys;
+
+	TSharedPtr<FUICommandInfo> BufferVisibleCurves;
+	TSharedPtr<FUICommandInfo> ApplyBufferedCurves;
+
+	// User Filtering
+	TSharedPtr<FUICommandInfo> OpenUserImplementableFilterWindow;
+
+	// Axis Snapping
+	TSharedPtr<FUICommandInfo> SetAxisSnappingNone;
+	TSharedPtr<FUICommandInfo> SetAxisSnappingHorizontal;
+	TSharedPtr<FUICommandInfo> SetAxisSnappingVertical;
+
+	// Time Management
+	TSharedPtr<FUICommandInfo> StepToNextKey;
+	TSharedPtr<FUICommandInfo> StepToPreviousKey;
+	TSharedPtr<FUICommandInfo> StepForward;
+	TSharedPtr<FUICommandInfo> StepBackward;
+	TSharedPtr<FUICommandInfo> JumpToStart;
+	TSharedPtr<FUICommandInfo> JumpToEnd;
 
 public:
 	virtual void RegisterCommands() override;
diff --git a/Engine/Source/Editor/CurveEditor/Public/CurveEditorHelpers.h b/Engine/Source/Editor/CurveEditor/Public/CurveEditorHelpers.h
new file mode 100644
--- /dev/null
+++ b/Engine/Source/Editor/CurveEditor/Public/CurveEditorHelpers.h
@@ -0,0 +1,18 @@
+// Copyright 1998-2019 Epic Games, Inc. All Rights Reserved.
+
+#pragma once
+
+#include "CurveEditorScreenSpace.h"
+
+namespace CurveEditor
+{
+
+	FVector2D ComputeScreenSpaceTangentOffset(const FCurveEditorScreenSpace& CurveSpace, float Tangent, float Weight);
+
+	void TangentAndWeightFromOffset(const FCurveEditorScreenSpace& CurveSpace, const FVector2D& TangentOffset, float& OutTangent, float& OutWeight);
+
+	FVector2D GetVectorFromSlopeAndLength(float Slope, float Length);
+
+	void ConstructYGridLines(const FCurveEditorScreenSpace& ViewSpace, uint8 InMinorDivisions, TArray<float>& OutMajorGridLines, TArray<float>& OutMinorGridLines, FText GridLineLabelFormatY, TArray<FText>* OutMajorGridLabels);
+
+} // namespace CurveEditor
\ No newline at end of file
diff --git a/Engine/Source/Editor/CurveEditor/Public/CurveEditorScreenSpace.h b/Engine/Source/Editor/CurveEditor/Public/CurveEditorScreenSpace.h
--- a/Engine/Source/Editor/CurveEditor/Public/CurveEditorScreenSpace.h
+++ b/Engine/Source/Editor/CurveEditor/Public/CurveEditorScreenSpace.h
@@ -2,25 +2,28 @@
 
 #pragma once
 
-#include "CoreMinimal.h"
-#include "Misc/FrameRate.h"
+#include "CoreTypes.h"
+#include "Math/Plane.h"
+#include "Math/Matrix.h"
+#include "Math/TransformCalculus2D.h"
 
-template<typename> class TRange;
-struct FGeometry;
 
 /**
  * Utility struct used for converting to/from curve editor screen space
  */
-struct FCurveEditorScreenSpace
+struct FCurveEditorScreenSpaceH
 {
 	/**
 	 * Construction from a physical size, and input/output range
 	 */
-	FCurveEditorScreenSpace(FVector2D InPixelSize, float InInputMin, float InInputMax, float InOutputMin, float InOutputMax)
-		: PixelSize(InPixelSize)
+	FCurveEditorScreenSpaceH(float InPixelWidth, double InInputMin, double InInputMax)
+		: PixelWidth(FMath::Max(InPixelWidth, 1.f))
 		, InputMin(InInputMin), InputMax(InInputMax)
-		, OutputMin(InOutputMin), OutputMax(InOutputMax)
 	{
+		if (!ensure(InputMax > InputMin))
+		{
+			InputMax = InputMin + KINDA_SMALL_NUMBER;
+		}
 	}
 
 public:
@@ -37,54 +40,155 @@ public:
 		return (InSeconds - InputMin) * PixelsPerInput();
 	}
 
-	/** Convert a vertical screen position in slate units to a value */
-	FORCEINLINE double ScreenToValue(float ScreenPosition) const
+public:
+
+	/** Retrieve the number of slate units per input value */
+	FORCEINLINE float PixelsPerInput() const
 	{
-		return OutputMin + (PixelSize.Y - ScreenPosition) / PixelsPerOutput();
+		float InputDiff = FMath::Max(InputMax - InputMin, 1e-10);
+		return PixelWidth / InputDiff;
 	}
 
-	/** Convert a value to a vertical screen position in slate units */
-	FORCEINLINE float ValueToScreen(double InValue) const
+public:
+
+	/** Get the minimum input value displayed on the screen */
+	FORCEINLINE double GetInputMin() const { return InputMin; }
+	/** Get the maximum input value displayed on the screen */
+	FORCEINLINE double GetInputMax() const { return InputMax; }
+	/** Get the physical width of the screen */
+	FORCEINLINE float GetPhysicalWidth() const { return PixelWidth; }
+
+public:
+
+	/**
+	 * Transform this screen space into a curve space using the specified transform
+	 */
+	FCurveEditorScreenSpaceH ToCurveSpace(const FTransform2D& CurveTransform) const
 	{
-		return PixelSize.Y - (InValue - OutputMin) * PixelsPerOutput();
+		const FVector2D& T = CurveTransform.GetTranslation();
+
+		float m00, m01, m10, m11;
+		CurveTransform.GetMatrix().GetMatrix(m00, m01, m10, m11);
+
+		double NewInputMin  = InputMin * m00 + T.X;
+		double NewInputMax  = InputMax * m00 + T.X;
+
+		return FCurveEditorScreenSpaceH(PixelWidth, NewInputMin, NewInputMax);
+	}
+
+private:
+
+	float PixelWidth;
+	double InputMin, InputMax;
+};
+
+
+/**
+ * Utility struct used for converting to/from curve editor screen space
+ */
+struct FCurveEditorScreenSpaceV
+{
+	/**
+	 * Construction from a physical size, and input/output range
+	 */
+	FCurveEditorScreenSpaceV(float InPixelHeight, double InOutputMin, double InOutputMax)
+		: PixelHeight(FMath::Max(InPixelHeight, 1.f))
+		, OutputMin(InOutputMin), OutputMax(InOutputMax)
+	{
+		if (!ensure(OutputMax > OutputMin))
+		{
+			OutputMax = OutputMin + KINDA_SMALL_NUMBER;
+		}
 	}
 
 public:
 
-	/** Retrieve the number of slate units per input value */
-	FORCEINLINE float PixelsPerInput() const
+	/** Convert a vertical screen position in slate units to a value */
+	FORCEINLINE double ScreenToValue(float ScreenPosition) const
 	{
-		float InputDiff = FMath::Max(InputMax - InputMin, 1e-10);
-		return PixelSize.X / InputDiff;
+		return OutputMin + (PixelHeight - ScreenPosition) / PixelsPerOutput();
 	}
 
+	/** Convert a value to a vertical screen position in slate units */
+	float ValueToScreen(double InValue) const
+	{
+		return (PixelHeight - (InValue - OutputMin) * PixelsPerOutput());
+	}
+
+public:
+
 	/** Retrieve the number of slate units per output value */
 	FORCEINLINE float PixelsPerOutput() const
 	{
 		float OutputDiff = FMath::Max(OutputMax - OutputMin, 1e-10);
-		return PixelSize.Y / OutputDiff;
+		return PixelHeight / OutputDiff;
 	}
 
 public:
 
-	/** Get the minimum input value displayed on the screen */
-	FORCEINLINE float GetInputMin() const { return InputMin; }
-	/** Get the maximum input value displayed on the screen */
-	FORCEINLINE float GetInputMax() const { return InputMax; }
 	/** Get the minimum output value displayed on the screen */
-	FORCEINLINE float GetOutputMin() const { return OutputMin; }
+	FORCEINLINE double GetOutputMin() const { return OutputMin; }
 	/** Get the maximum output value displayed on the screen */
-	FORCEINLINE float GetOutputMax() const { return OutputMax; }
-	/** Get the physical size of the screen */
-	FORCEINLINE FVector2D GetPhysicalSize() const { return PixelSize; }
-	/** Get the physical width of the screen */
-	FORCEINLINE float GetPhysicalWidth() const { return PixelSize.X; }
+	FORCEINLINE double GetOutputMax() const { return OutputMax; }
 	/** Get the physical height of the screen */
-	FORCEINLINE float GetPhysicalHeight() const { return PixelSize.Y; }
+	FORCEINLINE float GetPhysicalHeight() const { return PixelHeight; }
+
+public:
+
+	/**
+	 * Transform this screen space into a curve space using the specified transform
+	 */
+	FCurveEditorScreenSpaceV ToCurveSpace(const FTransform2D& CurveTransform) const
+	{
+		const FVector2D& T = CurveTransform.GetTranslation();
+
+		float m00, m01, m10, m11;
+		CurveTransform.GetMatrix().GetMatrix(m00, m01, m10, m11);
+
+		double NewOutputMin = OutputMin * m11 + T.Y;
+		double NewOutputMax = OutputMax * m11 + T.Y;
+
+		return FCurveEditorScreenSpaceV(PixelHeight, NewOutputMin, NewOutputMax);
+	}
 
 private:
 
-	FVector2D PixelSize;
-	double InputMin, InputMax;
+	float PixelHeight;
 	double OutputMin, OutputMax;
+};
+
+/**
+ * Utility struct used for converting to/from curve editor screen space
+ */
+struct FCurveEditorScreenSpace : FCurveEditorScreenSpaceH, FCurveEditorScreenSpaceV
+{
+	/**
+	 * Construction from a physical size, and input/output range
+	 */
+	FCurveEditorScreenSpace(FVector2D InPixelSize, double InInputMin, double InInputMax, double InOutputMin, double InOutputMax)
+		: FCurveEditorScreenSpaceH(InPixelSize.X, InInputMin, InInputMax)
+		, FCurveEditorScreenSpaceV(InPixelSize.Y, InOutputMin, InOutputMax)
+	{}
+
+public:
+
+	/** Get the physical size of the screen */
+	FORCEINLINE FVector2D GetPhysicalSize() const { return FVector2D(GetPhysicalWidth(), GetPhysicalHeight()); }
+
+public:
+
+	/**
+	 * Transform this screen space into a curve space using the specified transform
+	 */
+	FCurveEditorScreenSpace ToCurveSpace(const FTransform2D& CurveTransform) const
+	{
+		FCurveEditorScreenSpace New = *this;
+		FCurveEditorScreenSpaceH& H = static_cast<FCurveEditorScreenSpaceH&>(New);
+		FCurveEditorScreenSpaceV& V = static_cast<FCurveEditorScreenSpaceV&>(New);
+
+		H = H.ToCurveSpace(CurveTransform);
+		V = V.ToCurveSpace(CurveTransform);
+
+		return New;
+	}
 };
\ No newline at end of file
diff --git a/Engine/Source/Editor/CurveEditor/Public/CurveEditorSelection.h b/Engine/Source/Editor/CurveEditor/Public/CurveEditorSelection.h
--- a/Engine/Source/Editor/CurveEditor/Public/CurveEditorSelection.h
+++ b/Engine/Source/Editor/CurveEditor/Public/CurveEditorSelection.h
@@ -16,7 +16,7 @@ struct FCurvePointHandle;
  * A set of key handles implemented as a sorted array for transparent passing to TArrayView<> APIs.
  * Lookup is achieved via binary search: O(log(n)).
  */
-struct FKeyHandleSet
+struct CURVEEDITOR_API FKeyHandleSet
 {
 	/**
 	 * Add a new key handle to this set
@@ -124,7 +124,6 @@ public:
 	 * Add key handles to this selection, changing the selection type if necessary.
 	 */
 	void Add(FCurveModelID CurveID, ECurvePointType PointType, TArrayView<const FKeyHandle> Keys);
-
 public:
 
 	/**
diff --git a/Engine/Source/Editor/CurveEditor/Public/CurveEditorSettings.h b/Engine/Source/Editor/CurveEditor/Public/CurveEditorSettings.h
--- a/Engine/Source/Editor/CurveEditor/Public/CurveEditorSettings.h
+++ b/Engine/Source/Editor/CurveEditor/Public/CurveEditorSettings.h
@@ -9,25 +9,37 @@
 
 /** Defines visibility states for the tangents in the curve editor. */
 UENUM()
-namespace ECurveEditorTangentVisibility
+enum class ECurveEditorTangentVisibility : uint8
 {
-	enum Type
-	{
-		/** All tangents should be visible. */
-		AllTangents,
-		/** Only tangents from selected keys should be visible. */
-		SelectedKeys,
-		/** Don't display tangents. */
-		NoTangents
-	};
-}
+	/** All tangents should be visible. */
+	AllTangents,
+	/** Only tangents from selected keys should be visible. */
+	SelectedKeys,
+	/** Don't display tangents. */
+	NoTangents
+};
+
+/** Defines the position to center the zoom about in the curve editor. */
+UENUM()
+enum class ECurveEditorZoomPosition : uint8
+{
+	/** Current Time. */
+	CurrentTime,
+
+	/** Mouse Position. */
+	MousePosition,
+};
 
 /** Serializable options for curve editor. */
 UCLASS(config=EditorPerProjectUserSettings)
 class CURVEEDITOR_API UCurveEditorSettings : public UObject
 {
 public:
-	GENERATED_UCLASS_BODY()
+	GENERATED_BODY()
+
+	DECLARE_MULTICAST_DELEGATE(FOnCurveEditorCurveVisibilityChanged);
+
+	UCurveEditorSettings();
 
 	/** Gets whether or not the curve editor auto frames the selected curves. */
 	bool GetAutoFrameCurveEditor() const;
@@ -40,9 +52,14 @@ public:
 	void SetShowCurveEditorCurveToolTips(bool InbShowCurveEditorCurveToolTips);
 
 	/** Gets the current tangent visibility. */
-	ECurveEditorTangentVisibility::Type GetTangentVisibility() const;
+	ECurveEditorTangentVisibility GetTangentVisibility() const;
 	/** Sets the current tangent visibility. */
-	void SetTangentVisibility(ECurveEditorTangentVisibility::Type InTangentVisibility);
+	void SetTangentVisibility(ECurveEditorTangentVisibility InTangentVisibility);
+
+	/** Get zoom in/out position (mouse position or current time). */
+	ECurveEditorZoomPosition GetZoomPosition() const;
+	/** Set zoom in/out position (mouse position or current time). */
+	void SetZoomPosition(ECurveEditorZoomPosition InZoomPosition);
 
 protected:
 	UPROPERTY( config, EditAnywhere, Category="Curve Editor" )
@@ -52,5 +69,11 @@ protected:
 	bool bShowCurveEditorCurveToolTips;
 
 	UPROPERTY( config, EditAnywhere, Category="Curve Editor" )
-	TEnumAsByte<ECurveEditorTangentVisibility::Type> TangentVisibility;
+	ECurveEditorTangentVisibility TangentVisibility;
+
+	UPROPERTY( config, EditAnywhere, Category="Curve Editor")
+	ECurveEditorZoomPosition ZoomPosition;
+
+	/** Callback for when the Curve Editor Visibility setting changes. */
+	FOnCurveEditorCurveVisibilityChanged OnCurveEditorCurveVisibilityChanged;
 };
diff --git a/Engine/Source/Editor/CurveEditor/Public/CurveEditorSnapMetrics.h b/Engine/Source/Editor/CurveEditor/Public/CurveEditorSnapMetrics.h
--- a/Engine/Source/Editor/CurveEditor/Public/CurveEditorSnapMetrics.h
+++ b/Engine/Source/Editor/CurveEditor/Public/CurveEditorSnapMetrics.h
@@ -3,6 +3,7 @@
 #pragma once
 
 #include "Misc/FrameRate.h"
+#include "Math/Axis.h"
 
 /**
  * Utility struct that acts as a cache of the current snapping metrics for the curve editor
@@ -31,12 +32,59 @@ struct FCurveEditorSnapMetrics
 	/** Snap the specified input time to the input snap rate if necessary */
 	FORCEINLINE double SnapInputSeconds(double InputTime)
 	{
-		return bSnapInputValues ? (InputTime * InputSnapRate).RoundToFrame() / InputSnapRate : InputTime;
+		return bSnapInputValues && InputSnapRate.IsValid() ? (InputTime * InputSnapRate).RoundToFrame() / InputSnapRate : InputTime;
 	}
 
 	/** Snap the specified output value to the output snap interval if necessary */
 	FORCEINLINE double SnapOutput(double OutputValue)
 	{
-		return bSnapOutputValues ? FMath::RoundToDouble(OutputValue / OutputSnapInterval) * OutputSnapInterval : OutputValue;
+		return bSnapOutputValues && OutputSnapInterval != 0.f ? FMath::RoundToDouble(OutputValue / OutputSnapInterval) * OutputSnapInterval : OutputValue;
 	}
 };
+
+/**
+ * Utility struct that acts as a way to control snapping to a specific axis based on UI settings, or shift key.
+ */
+struct FCurveEditorAxisSnap
+{
+	/** Can be set to either X, Y, or None to control which axis GetSnappedPosition snaps to. User can override None by pressing shift. */
+	EAxisList::Type RestrictedAxisList;
+
+	FCurveEditorAxisSnap()
+	{
+		RestrictedAxisList = EAxisList::None;
+	}
+
+	FVector2D GetSnappedPosition(FVector2D InitialPosition, FVector2D CurrentPosition, const FPointerEvent& MouseEvent, const bool bIgnoreAxisLock = false)
+	{
+		check(RestrictedAxisList == EAxisList::Type::None || RestrictedAxisList == EAxisList::Type::X || RestrictedAxisList == EAxisList::Type::Y);
+
+		// If we're ignoring axis lock (such as for UI) we allow them to use shift anyways
+		bool bCanUseShift = bIgnoreAxisLock || RestrictedAxisList == EAxisList::Type::None;
+
+		FVector2D MouseLockVector = FVector2D::UnitVector;
+		if (bCanUseShift && MouseEvent.IsShiftDown())
+		{
+			// Figure out which axis is currently greater and snap to that 
+			FVector2D DragDelta = CurrentPosition - InitialPosition;
+			if (FMath::Abs(DragDelta.X) > FMath::Abs(DragDelta.Y))
+			{
+				MouseLockVector.Y = 0;
+			}
+			else
+			{
+				MouseLockVector.X = 0;
+			}
+		}
+		else if (!bIgnoreAxisLock && RestrictedAxisList == EAxisList::Type::X)
+		{
+			MouseLockVector.Y = 0;
+		}
+		else if (!bIgnoreAxisLock && RestrictedAxisList == EAxisList::Type::Y)
+		{
+			MouseLockVector.X = 0;
+		}
+
+		return InitialPosition + (CurrentPosition - InitialPosition) * MouseLockVector;
+	}
+};
\ No newline at end of file
diff --git a/Engine/Source/Editor/CurveEditor/Public/CurveEditorTypes.h b/Engine/Source/Editor/CurveEditor/Public/CurveEditorTypes.h
--- a/Engine/Source/Editor/CurveEditor/Public/CurveEditorTypes.h
+++ b/Engine/Source/Editor/CurveEditor/Public/CurveEditorTypes.h
@@ -5,6 +5,9 @@
 #include "CoreTypes.h"
 #include "Templates/TypeHash.h"
 #include "Curves/KeyHandle.h"
+#include "Containers/BitArray.h"
+#include "Containers/Array.h"
+#include "Templates/SharedPointer.h"
 
 /**
  * Enum for representing the type of a key point in the curve editor
@@ -14,11 +17,30 @@ enum class ECurvePointType : uint8
 	Key, ArriveTangent, LeaveTangent
 };
 
+enum class ECurveEditorTreeSelectionState : uint8
+{
+	None, Explicit, ImplicitChild
+};
+
+/** Enumeration identifying a single specific view type, or a combination thereof */
+enum class ECurveEditorViewID : uint64
+{
+	Invalid      = 0,
+
+	Absolute     = 1 << 0,
+	Normalized   = 1 << 1,
+	Stacked      = 1 << 2,
+
+	CUSTOM_START = 1 << 3,
+	ANY_BUILT_IN = Absolute | Normalized | Stacked,
+};
+ENUM_CLASS_FLAGS(ECurveEditorViewID);
+
 
 /**
  * A unique identifier for a curve model existing on a curve editor
  */
-struct FCurveModelID
+struct CURVEEDITOR_API FCurveModelID
 {
 	/**
 	 * Generate a new curve model ID
@@ -33,6 +55,22 @@ struct FCurveModelID
 		return A.ID == B.ID;
 	}
 
+	/**
+	 * Check two IDs for inequality
+	 */
+	FORCEINLINE friend bool operator!=(FCurveModelID A, FCurveModelID B)
+	{
+		return A.ID != B.ID;
+	}
+
+	/**
+	 * Test whether A is less than B
+	 */
+	FORCEINLINE friend bool operator<(FCurveModelID A, FCurveModelID B)
+	{
+		return A.ID < B.ID;
+	}
+
 	/**
 	 * Hash a curve model ID
 	 */
@@ -41,6 +79,11 @@ struct FCurveModelID
 		return GetTypeHash(In.ID);
 	}
 
+	FCurveModelID(const FCurveModelID& InOther)
+		: ID(InOther.ID)
+	{
+	}
+
 private:
 	FCurveModelID() {}
 
@@ -66,4 +109,46 @@ struct FCurvePointHandle
 	ECurvePointType PointType;
 	/** The key handle for the underlying key */
 	FKeyHandle KeyHandle;
+};
+
+struct FCurveEditorTreeItemID
+{
+	FCurveEditorTreeItemID()
+		: Value(0)
+	{}
+
+	friend bool operator==(FCurveEditorTreeItemID A, FCurveEditorTreeItemID B)
+	{
+		return A.Value == B.Value;
+	}
+
+	friend bool operator!=(FCurveEditorTreeItemID A, FCurveEditorTreeItemID B)
+	{
+		return A.Value != B.Value;
+	}
+
+	friend uint32 GetTypeHash(FCurveEditorTreeItemID ID)
+	{
+		return GetTypeHash(ID.Value);
+	}
+
+	static FCurveEditorTreeItemID Invalid()
+	{
+		return FCurveEditorTreeItemID();
+	}
+
+	bool IsValid() const
+	{
+		return Value != 0;
+	}
+
+	uint32 GetValue() const
+	{
+		return Value;
+	}
+
+protected:
+
+	friend class FCurveEditorTree;
+	uint32 Value;
 };
\ No newline at end of file
diff --git a/Engine/Source/Editor/CurveEditor/Public/CurveModel.h b/Engine/Source/Editor/CurveEditor/Public/CurveModel.h
--- a/Engine/Source/Editor/CurveEditor/Public/CurveModel.h
+++ b/Engine/Source/Editor/CurveEditor/Public/CurveModel.h
@@ -8,6 +8,7 @@
 #include "Containers/ArrayView.h"
 #include "Curves/RichCurve.h"
 #include "CurveEditorTypes.h"
+#include "Math/TransformCalculus2D.h"
 
 struct FKeyHandle;
 struct FKeyDrawInfo;
@@ -16,24 +17,26 @@ struct FKeyPosition;
 struct FKeyAttributes;
 struct FCurveAttributes;
 struct FCurveEditorScreenSpace;
+struct FCurveModelID;
 
 class FName;
 class SWidget;
 class FCurveEditor;
 class UObject;
+class SCurveEditorView;
 
 enum class ECurvePointType : uint8;
 
-
 /**
  * Class that models an underlying curve data structure through a generic abstraction that the curve editor understands.
  */
-class FCurveModel
+class CURVEEDITOR_API FCurveModel
 {
 public:
 
 	FCurveModel()
 		: Color(FLinearColor::White)
+		, SupportedViews(ECurveEditorViewID::ANY_BUILT_IN)
 	{}
 
 	virtual ~FCurveModel()
@@ -53,9 +56,10 @@ public:
 	 * Draw the curve for the specified curve editor by populating an array with points on the curve between which lines should be drawn
 	 *
 	 * @param CurveEditor             Reference to the curve editor that is drawing the curve. Can be used to cull the interpolating points to the visible region.
+	 * @param ScreenSpace			  A transform which indicates the use case for the drawn curve. This lets you simplify curves based on their screenspace representation.
 	 * @param OutInterpolatingPoints  Array to populate with points (time, value) that lie on the curve.
 	 */
-	virtual void DrawCurve(const FCurveEditor& CurveEditor, TArray<TTuple<double, double>>& InterpolatingPoints) const = 0;
+	virtual void DrawCurve(const FCurveEditor& CurveEditor, const FCurveEditorScreenSpace& ScreenSpace, TArray<TTuple<double, double>>& InterpolatingPoints) const = 0;
 
 	/**
 	 * Retrieve all keys that lie in the specified time and value range
@@ -105,9 +109,10 @@ public:
 	 * Populate the specified draw info structure with data describing how to draw the specified point type
 	 *
 	 * @param PointType              The type of point to be drawn
+	 * @param InKeyHandle			 The specific key (if possible, otherwise FKeyHandle::Invalid()) to get the info for.
 	 * @param OutDrawInfo            Data structure to be populated with draw info for this type of point
 	 */
-	virtual void GetKeyDrawInfo(ECurvePointType PointType, FKeyDrawInfo& OutDrawInfo) const = 0;
+	virtual void GetKeyDrawInfo(ECurvePointType PointType, const FKeyHandle InKeyHandle, FKeyDrawInfo& OutDrawInfo) const = 0;
 
 	/** Get range of input time.
 	* @param MinTime Minimum Time
@@ -122,6 +127,19 @@ public:
 	*/
 	virtual void GetValueRange(double& MinValue, double& MaxValue) const = 0;
 
+	/** Get the number of keys
+	* @param The number of keys
+	*/
+	virtual int32 GetNumKeys() const = 0;
+
+	/** Get neighboring keys given the key handle
+	 *
+	 * @param InKeyHandle The key handle to get neighboring keys for
+	 * @param OutPreviousKeyHandle The previous key handle
+     * @param OutNextKeyHandle The next key handle
+	 */
+	virtual void GetNeighboringKeys(const FKeyHandle InKeyHandle, TOptional<FKeyHandle>& OutPreviousKeyHandle, TOptional<FKeyHandle>& OutNextKeyHandle) const = 0;
+
 	/**
 	 * Evaluate this curve at the specified time
 	 *
@@ -184,34 +202,84 @@ public:
 	virtual void CreateKeyProxies(TArrayView<const FKeyHandle> InKeyHandles, TArrayView<UObject*> OutObjects)
 	{}
 
-public:
-
 	/**
 	 * Helper function for assigning a the same attributes to a number of keys
 	 */
-	CURVEEDITOR_API void SetKeyAttributes(TArrayView<const FKeyHandle> InKeys, const FKeyAttributes& InAttributes);
+	void SetKeyAttributes(TArrayView<const FKeyHandle> InKeys, const FKeyAttributes& InAttributes);
 
 	/**
 	 * Helper function for adding a single key to this curve
 	 */
-	CURVEEDITOR_API TOptional<FKeyHandle> AddKey(const FKeyPosition& NewKeyPosition, const FKeyAttributes& InAttributes);
+	TOptional<FKeyHandle> AddKey(const FKeyPosition& NewKeyPosition, const FKeyAttributes& InAttributes);
 
 public:
 
 	/**
-	 * Access this curve's display name
+	 * Access this curve's short display name. This is useful when there are other UI elements which describe
+	 * enough context about the curve that a long name is not needed (ie: Showing just "X" because other
+	 * UI elements give the object/group context).
 	 */
-	FORCEINLINE FText GetDisplayName() const
+	FORCEINLINE FText GetShortDisplayName() const
 	{
-		return DisplayName;
+		return ShortDisplayName;
+	}
+
+	/**
+	 * Assign a short display name for this curve
+	 */
+	FORCEINLINE void SetShortDisplayName(const FText& InDisplayName)
+	{
+		ShortDisplayName = InDisplayName;
+	}
+
+	/**
+	 * Access this curve's long display name. This is useful when you want more context about
+	 * the curve, such as the object it belongs to, or the group (ie: "Floor.Transform.X") instead
+	 * of just "X" or "Transform.X"
+	 */
+	FORCEINLINE FText GetLongDisplayName() const
+	{
+		// For convenience fall back to the short display name if they fail to specify a long one.
+		if (LongDisplayName.IsEmptyOrWhitespace())
+		{
+			return GetShortDisplayName();
+		}
+
+		return LongDisplayName;
 	}
 
 	/**
-	 * Assign a display name for this curve
+	 * Assign a long display name for this curve used in contexts where additional context is useful.
 	 */
-	FORCEINLINE void SetDisplayName(FText InDisplayName)
+	FORCEINLINE void SetLongDisplayName(const FText& InLongDisplayName)
 	{
-		DisplayName = InDisplayName;
+		LongDisplayName = InLongDisplayName;
+	}
+
+	/**
+	 * This is an internal name used to try to match different curves with each other. When saving
+	 * and later restoring curves on a different set of curves we need a name that gives enough context
+	 * to match them up by intention, and not long or short name. For example, a curve might have a short
+	 * name of "X", and a long name of "Floor.Transform.Location.X". If you wanted to copy a set of transform
+	 * curves and paste them onto another transform, we use this context to match the names together to ensure
+	 * your Transform.X gets applied to the other Transform.X - in this example the intention is
+	 * for the curve to represent a "Location.X" (so it should be pasteable on any other curve which says
+	 * their context is a "Location.X" as well). This is more reliable and more flexible than relying on
+	 * short display names (not enough context in the case of seeing Location.X, and Scale.X) and better
+	 * than relying on long display names (too much context and no reliable way to substring them).
+	 */
+	FORCEINLINE FString GetIntentionName() const
+	{
+		return IntentionName;
+	}
+
+	/**
+	 * Assign an intention name for this curve which is used internally when applying one curve to another in situations where
+	 * multiple curves are visible.
+	 */
+	FORCEINLINE void SetIntentionName(const FString& InIntentionName)
+	{
+		IntentionName = InIntentionName;
 	}
 
 	/**
@@ -230,11 +298,28 @@ public:
 		Color = InColor;
 	}
 
+	/**
+	 * Retrieve this curve's supported views
+	 */
+	ECurveEditorViewID GetSupportedViews() const
+	{
+		return SupportedViews;
+	}
+
 protected:
 
-	/** This curve's display name */
-	FText DisplayName;
+	/** This curve's short display name. Used in situations where other mechanisms provide enough context about what the curve is (such as "X") */
+	FText ShortDisplayName;
+
+	/** This curve's long display name. Used in situations where the UI doesn't provide enough context about what the curve is otherwise (such as "Floor.Transform.X") */
+	FText LongDisplayName;
+
+	/** This curve's intention (such as Transform.X or Scale.X). Used internally to match up curves when saving/restoring curves between different objects. */
+	FString IntentionName;
 
 	/** This curve's display color */
 	FLinearColor Color;
+
+	/** A set of views supported by this curve */
+	ECurveEditorViewID SupportedViews;
 };
