Rearranged the Sequencer curve editor toolbar to match design.

diff --git a/Engine/Source/Editor/CurveEditor/Private/SCurveEditorPanel.cpp b/Engine/Source/Editor/CurveEditor/Private/SCurveEditorPanel.cpp
--- a/Engine/Source/Editor/CurveEditor/Private/SCurveEditorPanel.cpp
+++ b/Engine/Source/Editor/CurveEditor/Private/SCurveEditorPanel.cpp
@@ -989,19 +989,6 @@ TSharedPtr<FExtender> SCurveEditorPanel::GetToolbarExtender()
 	{
 		static void FillToolbar(FToolBarBuilder& ToolBarBuilder, TSharedRef<SCurveKeyDetailPanel> InKeyDetailsPanel, TSharedRef<SCurveEditorPanel> InEditorPanel)
 		{
-
-			ToolBarBuilder.BeginSection("Tools");
-			{
-				ToolBarBuilder.AddToolBarButton(FCurveEditorCommands::Get().DeactivateCurrentTool);
-			}
-			ToolBarBuilder.EndSection();
-
-			ToolBarBuilder.BeginSection("Key Details");
-			{
-				ToolBarBuilder.AddWidget(InKeyDetailsPanel);
-			}
-			ToolBarBuilder.EndSection();
-
 			ToolBarBuilder.BeginSection("View");
 			{
 				// Dropdown Menu for choosing your viewing mode
@@ -1049,6 +1036,18 @@ TSharedPtr<FExtender> SCurveEditorPanel::GetToolbarExtender()
 			}
 			ToolBarBuilder.EndSection();
 
+			ToolBarBuilder.BeginSection("Key Details");
+			{
+				ToolBarBuilder.AddWidget(InKeyDetailsPanel);
+			}
+			ToolBarBuilder.EndSection();
+
+			ToolBarBuilder.BeginSection("Tools");
+			{
+				ToolBarBuilder.AddToolBarButton(FCurveEditorCommands::Get().DeactivateCurrentTool);
+			}
+			ToolBarBuilder.EndSection();
+
 			ToolBarBuilder.BeginSection("Adjustment");
 			{
 				// Toggle Button for Time Snapping
