Removed SNumericList from the generic Slate widgets and moved it to the Curve Editor as SCurveValueSnapList. This widget was only in use by the curve editor and the curve editor may be phasing its use case out soon.

diff --git a/Engine/Source/Editor/CurveEditor/Private/SCurveEditorPanel.cpp b/Engine/Source/Editor/CurveEditor/Private/SCurveEditorPanel.cpp
--- a/Engine/Source/Editor/CurveEditor/Private/SCurveEditorPanel.cpp
+++ b/Engine/Source/Editor/CurveEditor/Private/SCurveEditorPanel.cpp
@@ -32,7 +32,7 @@
 #include "Filters/CurveEditorFilterBase.h"
 #include "Filters/CurveEditorBakeFilter.h"
 #include "Filters/CurveEditorReduceFilter.h"
-#include "Widgets/Input/SNumericList.h"
+#include "SCurveValueSnapList.h"
 #include "Widgets/SFrameRatePicker.h"
 #include "CommonFrameRates.h"
 #include "CurveEditorViewRegistry.h"
@@ -1188,20 +1188,20 @@ FText SCurveEditorPanel::GetTimeSnapMenuTooltip() const
 
 TSharedRef<SWidget> SCurveEditorPanel::MakeValueSnapMenu()
 {
-	TArray<SNumericList<float>::FNamedValue> ValueSnapAmounts;
+	TArray<SCurveValueSnapList<float>::FNamedValue> ValueSnapAmounts;
 	// SnapValues.Add( SNumericDropDown<float>::FNamedValue( 0.001f, LOCTEXT( "Snap_OneThousandth", "0.001" ), LOCTEXT( "SnapDescription_OneThousandth", "Set snap to 1/1000th" ) ) );
 	//SnapValues.Add( SNumericDropDown<float>::FNamedValue( 0.01f, LOCTEXT( "Snap_OneHundredth", "0.01" ), LOCTEXT( "SnapDescription_OneHundredth", "Set snap to 1/100th" ) ) );
-	ValueSnapAmounts.Add(SNumericList<float>::FNamedValue(0.1f, LOCTEXT("Snap_OneTenth", "0.1"), LOCTEXT("SnapDescription_OneTenth", "Set snap to 1/10th")));
-	ValueSnapAmounts.Add(SNumericList<float>::FNamedValue(0.5f, LOCTEXT("Snap_OneHalf", "0.5"), LOCTEXT("SnapDescription_OneHalf", "Set snap to 1/2")));
-	ValueSnapAmounts.Add(SNumericList<float>::FNamedValue(1.0f, LOCTEXT("Snap_One", "1"), LOCTEXT("SnapDescription_One", "Set snap to 1")));
-	ValueSnapAmounts.Add(SNumericList<float>::FNamedValue(2.0f, LOCTEXT("Snap_Two", "2"), LOCTEXT("SnapDescription_Two", "Set snap to 2")));
-	ValueSnapAmounts.Add(SNumericList<float>::FNamedValue(5.0f, LOCTEXT("Snap_Five", "5"), LOCTEXT("SnapDescription_Five", "Set snap to 5")));
-	ValueSnapAmounts.Add(SNumericList<float>::FNamedValue(10.0f, LOCTEXT("Snap_Ten", "10"), LOCTEXT("SnapDescription_Ten", "Set snap to 10")));
-	ValueSnapAmounts.Add(SNumericList<float>::FNamedValue(50.0f, LOCTEXT("Snap_Fifty", "50"), LOCTEXT("SnapDescription_50", "Set snap to 50")));
-	ValueSnapAmounts.Add(SNumericList<float>::FNamedValue(100.0f, LOCTEXT("Snap_OneHundred", "100"), LOCTEXT("SnapDescription_OneHundred", "Set snap to 100")));
+	ValueSnapAmounts.Add(SCurveValueSnapList<float>::FNamedValue(0.1f, LOCTEXT("Snap_OneTenth", "0.1"), LOCTEXT("SnapDescription_OneTenth", "Set snap to 1/10th")));
+	ValueSnapAmounts.Add(SCurveValueSnapList<float>::FNamedValue(0.5f, LOCTEXT("Snap_OneHalf", "0.5"), LOCTEXT("SnapDescription_OneHalf", "Set snap to 1/2")));
+	ValueSnapAmounts.Add(SCurveValueSnapList<float>::FNamedValue(1.0f, LOCTEXT("Snap_One", "1"), LOCTEXT("SnapDescription_One", "Set snap to 1")));
+	ValueSnapAmounts.Add(SCurveValueSnapList<float>::FNamedValue(2.0f, LOCTEXT("Snap_Two", "2"), LOCTEXT("SnapDescription_Two", "Set snap to 2")));
+	ValueSnapAmounts.Add(SCurveValueSnapList<float>::FNamedValue(5.0f, LOCTEXT("Snap_Five", "5"), LOCTEXT("SnapDescription_Five", "Set snap to 5")));
+	ValueSnapAmounts.Add(SCurveValueSnapList<float>::FNamedValue(10.0f, LOCTEXT("Snap_Ten", "10"), LOCTEXT("SnapDescription_Ten", "Set snap to 10")));
+	ValueSnapAmounts.Add(SCurveValueSnapList<float>::FNamedValue(50.0f, LOCTEXT("Snap_Fifty", "50"), LOCTEXT("SnapDescription_50", "Set snap to 50")));
+	ValueSnapAmounts.Add(SCurveValueSnapList<float>::FNamedValue(100.0f, LOCTEXT("Snap_OneHundred", "100"), LOCTEXT("SnapDescription_OneHundred", "Set snap to 100")));
 
 	TSharedRef<SWidget> OutputSnapWidget =
-		SNew(SNumericList<float>)
+		SNew(SCurveValueSnapList<float>)
 		.DropDownValues(ValueSnapAmounts)
 		.MinDesiredValueWidth(60)
 		.Value_Lambda([this]() -> float { return this->CurveEditor->OutputSnapIntervalAttribute.Get(); })
diff --git a/Engine/Source/Editor/CurveEditor/Public/SCurveValueSnapList.h b/Engine/Source/Editor/CurveEditor/Public/SCurveValueSnapList.h
new file mode 100644
--- /dev/null
+++ b/Engine/Source/Editor/CurveEditor/Public/SCurveValueSnapList.h
@@ -0,0 +1,189 @@
+// Copyright 1998-2019 Epic Games, Inc. All Rights Reserved.
+
+#pragma once
+
+#include "CoreMinimal.h"
+#include "Misc/Attribute.h"
+#include "Layout/Visibility.h"
+#include "Input/Events.h"
+#include "Widgets/SWidget.h"
+#include "Textures/SlateIcon.h"
+#include "Layout/Margin.h"
+#include "Widgets/DeclarativeSyntaxSupport.h"
+#include "Widgets/SCompoundWidget.h"
+#include "Widgets/SBoxPanel.h"
+#include "Framework/Commands/UIAction.h"
+#include "Framework/MultiBox/MultiBoxBuilder.h"
+#include "Widgets/Text/STextBlock.h"
+#include "Framework/Application/SlateApplication.h"
+#include "Widgets/Input/SEditableTextBox.h"
+#include "Widgets/Input/SComboButton.h"
+
+/** A widget which creates a list of pre-determined numbers and a spot to enter a custom number. Similar to SNumericList but without the dropdown menu. */
+template<typename NumericType>
+class SCurveValueSnapList : public SCompoundWidget
+{
+public:
+	DECLARE_DELEGATE_OneParam( FOnValueChanged, NumericType );
+
+public:
+	/** Represents a named numeric value for display in the drop down menu. */
+	class FNamedValue
+	{
+	public:
+		/** Creates a new FNamedValue 
+		  * @InValue The numeric value to be assigned
+		  * @InName The display text for the value in the UI 
+		  * @InDescription The description of the value used in tooltips or wherever a longer description is needed. */
+		FNamedValue( NumericType InValue, FText InName, FText InDescription )
+		{
+			Value = InValue;
+			Name = InName;
+			Description = InDescription;
+		}
+
+		NumericType GetValue() const
+		{
+			return Value;
+		}
+
+		FText GetName() const
+		{
+			return Name;
+		}
+
+		FText GetDescription() const
+		{
+			return Description;
+		}
+
+	private:
+		NumericType Value;
+		FText Name;
+		FText Description;
+	};
+
+public:
+	SLATE_BEGIN_ARGS( SCurveValueSnapList<NumericType> )
+		: _DropDownValues()
+		, _MinDesiredValueWidth( 40 )
+		, _bShowNamedValue(true)
+		{}
+
+		/** The values which are used to populate the drop down menu. */
+		SLATE_ARGUMENT( TArray<FNamedValue>, DropDownValues )
+		/** Controls the minimum width for the text box portion of the control. */
+		SLATE_ATTRIBUTE( float, MinDesiredValueWidth )
+		/** The value displayed by the control. */
+		SLATE_ATTRIBUTE( NumericType, Value )
+		/** If enabled, attempt to show the original name specified in the Drop Down values for a given entry instead of just a straight FText::AsNumber(NumericType). */
+		SLATE_ATTRIBUTE(bool, bShowNamedValue)
+		/** An optional header to prepend to the generated list. */
+		SLATE_ATTRIBUTE( FText, HeaderText)
+		/** The callback for when the value changes. */
+		SLATE_EVENT( FOnValueChanged, OnValueChanged )
+
+	SLATE_END_ARGS()
+
+	void Construct( const FArguments& InArgs )
+	{
+		DropDownValues = InArgs._DropDownValues;
+		Value = InArgs._Value;
+		OnValueChanged = InArgs._OnValueChanged;
+
+		FMenuBuilder MenuBuilder(true, nullptr);
+
+		// Allow an optional header text because it looks nicer and helps give context to the widget.
+		if (InArgs._HeaderText.IsSet())
+		{
+			MenuBuilder.BeginSection("Header", InArgs._HeaderText.Get());
+		}
+
+		for (const FNamedValue& NamedValue : DropDownValues)
+		{
+			AddMenuEntry(MenuBuilder, NamedValue);
+		}
+
+		// Add the custom entry box at the end. 
+		MenuBuilder.AddWidget(
+			SNew(SBox)
+			.HAlign(HAlign_Right)
+			.Padding(FMargin(8, 0, 0, 0))
+			[
+				SNew(SEditableTextBox)
+				.RevertTextOnEscape(true)
+				.SelectAllTextWhenFocused(true)
+				.Text(this, &SCurveValueSnapList<NumericType>::GetValueText)
+				.OnTextCommitted(this, &SCurveValueSnapList<NumericType>::ValueTextCommitted)
+				.MinDesiredWidth(InArgs._MinDesiredValueWidth)
+			],
+			NSLOCTEXT("SNumericList", "CustomNumberDisplayLabel", "Custom")
+		);
+
+		ChildSlot
+		[
+			MenuBuilder.MakeWidget()
+		];
+	}
+
+private:
+	void AddMenuEntry(FMenuBuilder& MenuBuilder, const FNamedValue& Info)
+	{
+		MenuBuilder.AddMenuEntry(
+			Info.GetName(),
+			Info.GetDescription(),
+			FSlateIcon(),
+			FUIAction(
+				FExecuteAction::CreateSP(this, &SCurveValueSnapList<NumericType>::SetValue, Info.GetValue()),
+				FCanExecuteAction(),
+				FIsActionChecked::CreateSP(this, &SCurveValueSnapList<NumericType>::IsSameValue, Info.GetValue())
+			),
+			NAME_None,
+			EUserInterfaceActionType::RadioButton
+		);
+	}
+
+	bool IsSameValue(NumericType InValue) const
+	{
+		return Value.Get() == InValue;
+	}
+
+	/** Get the value text for displaying in the custom entry box. We use a textbox here so we can show actual names from the original list if possible. */
+	FText GetValueText() const
+	{
+		if (bShowNamedValue.Get())
+		{
+			for ( FNamedValue DropDownValue : DropDownValues )
+			{
+				if (FMath::IsNearlyEqual(DropDownValue.GetValue(), Value.Get()))
+				{
+					return DropDownValue.GetName();
+				}
+			}
+		}
+
+		return FText::FromString(LexToSanitizedString(Value.Get()));
+	}
+
+	void ValueTextCommitted(const FText& InNewText, ETextCommit::Type InTextCommit)
+	{
+		if (InNewText.IsNumeric())
+		{
+			NumericType NewValue;
+			TTypeFromString<NumericType>::FromString(NewValue, *InNewText.ToString());
+			OnValueChanged.ExecuteIfBound(NewValue);
+		}
+	}
+
+	void SetValue( NumericType InValue )
+	{
+		FSlateApplication::Get().ClearKeyboardFocus( EFocusCause::Cleared );
+		OnValueChanged.ExecuteIfBound(InValue);
+	}
+
+private:
+	TArray<FNamedValue> DropDownValues;
+	TAttribute<bool> bShowNamedValue;
+	TAttribute<NumericType> Value;
+	FOnValueChanged OnValueChanged;
+};
diff --git a/Engine/Source/Runtime/Slate/Public/Widgets/Input/SNumericList.h b/Engine/Source/Runtime/Slate/Public/Widgets/Input/SNumericList.h
deleted file mode 100644
