Applying buffered curves now creates a transaction so it can be undone.

diff --git a/Engine/Source/Editor/CurveEditor/Private/CurveEditor.cpp b/Engine/Source/Editor/CurveEditor/Private/CurveEditor.cpp
--- a/Engine/Source/Editor/CurveEditor/Private/CurveEditor.cpp
+++ b/Engine/Source/Editor/CurveEditor/Private/CurveEditor.cpp
@@ -985,6 +985,8 @@ void FCurveEditor::ApplyBufferedCurveToTarget(const FBufferedCurve& BufferedCurv
 
 bool FCurveEditor::ApplyBufferedCurves(const TSet<FCurveModelID>& InCurvesToApplyTo)
 {
+	FScopedTransaction Transaction(LOCTEXT("ApplyBufferedCurves", "Apply Buffered Curves"));
+
 	// Each curve can specify an "Intention" name. This gives a little bit of context about how the curve is intended to be used,
 	// without locking anyone into a specific set of intentions. When you go to apply the buffered curves, for each curve that you
 	// want to apply it to, we can look in our stored curves to see if someone has the same intention. If there isn't a matching intention
@@ -1142,6 +1144,8 @@ bool FCurveEditor::ApplyBufferedCurves(const TSet<FCurveModelID>& InCurvesToAppl
 		FSlateNotificationManager::Get().AddNotification(FailInfo);
 	}
 
+	// No need to make a entry in the Undo/Redo buffer if it didn't apply anything.
+	Transaction.Cancel();
 	return false;
 }
 
