Curve Editor: Adding euler filter

diff --git a/Engine/Source/Editor/CurveEditor/Private/Filters/CurveEditorEulerFilter.cpp b/Engine/Source/Editor/CurveEditor/Private/Filters/CurveEditorEulerFilter.cpp
new file mode 100644
--- /dev/null
+++ b/Engine/Source/Editor/CurveEditor/Private/Filters/CurveEditorEulerFilter.cpp
@@ -0,0 +1,85 @@
+// Copyright 1998-2019 Epic Games, Inc. All Rights Reserved.
+
+#include "Filters/CurveEditorEulerFilter.h"
+#include "CoreMinimal.h"
+#include "ScopedTransaction.h"
+#include "Curves/KeyHandle.h"
+#include "CurveDataAbstraction.h"
+#include "CurveEditor.h"
+#include "CurveEditorSelection.h"
+#include "CurveModel.h"
+#include "Misc/FrameRate.h"
+#include "CurveEditorSnapMetrics.h"
+#include "Containers/ArrayView.h"
+
+void UCurveEditorEulerFilter::ApplyFilter_Impl(TSharedRef<FCurveEditor> InCurveEditor, const TMap<FCurveModelID, FKeyHandleSet>& InKeysToOperateOn, TMap<FCurveModelID, FKeyHandleSet>& OutKeysToSelect)
+{
+	TArray<FKeyHandle> KeyHandles;
+	TArray<FKeyHandle> KeyHandlesToModify;
+	TArray<FKeyPosition> SelectedKeyPositions;
+
+	TArray<FKeyPosition> NewKeyPositions;
+
+	for (const TTuple<FCurveModelID, FKeyHandleSet>& Pair : InKeysToOperateOn)
+	{
+		// first check if curve exists
+		FCurveModel* Curve = InCurveEditor->FindCurve(Pair.Key);
+		if (!Curve)
+		{
+			continue;
+		}
+
+		// then check if curve is for rotation value
+		if (!Curve->GetIntentionName().Contains("Rotation"))
+		{
+			continue;
+		}
+
+		KeyHandles.Reset(Pair.Value.Num());
+		KeyHandles.Append(Pair.Value.AsArray().GetData(), Pair.Value.Num());
+
+		// Get all the selected keys
+		SelectedKeyPositions.SetNum(KeyHandles.Num());
+		Curve->GetKeyPositions(KeyHandles, SelectedKeyPositions);
+
+		// Find the hull of the range of the selected keys
+		double MinKey = TNumericLimits<double>::Max(), MaxKey = TNumericLimits<double>::Lowest();
+		for (FKeyPosition Key : SelectedKeyPositions)
+		{
+			MinKey = FMath::Min(Key.InputValue, MinKey);
+			MaxKey = FMath::Max(Key.InputValue, MaxKey);
+		}
+
+		// Get all keys that exist between the time range
+		KeyHandles.Reset();
+		Curve->GetKeys(*InCurveEditor, MinKey, MaxKey, TNumericLimits<double>::Lowest(), TNumericLimits<double>::Max(), KeyHandles);
+
+		if (KeyHandles.Num() >= 2)
+		{
+			// get the selected key positions
+			SelectedKeyPositions.SetNum(KeyHandles.Num());
+			Curve->GetKeyPositions(KeyHandles, SelectedKeyPositions);
+			
+			// reset the old calculated key positions
+			NewKeyPositions.Reset();
+			KeyHandlesToModify.Reset();
+
+			for (int32 KeyIndex = 0; KeyIndex < KeyHandles.Num() - 1; ++KeyIndex)
+			{
+				// calculate the euler-filtered key
+				float NextKeyVal = SelectedKeyPositions[KeyIndex + 1].OutputValue;
+				FMath::WindRelativeAnglesDegrees(SelectedKeyPositions[KeyIndex].OutputValue, NextKeyVal);
+				SelectedKeyPositions[KeyIndex + 1].OutputValue = NextKeyVal;
+				
+				// create the new key position and add it and its handle to their respective lists
+				FKeyPosition FilteredKey(SelectedKeyPositions[KeyIndex + 1].InputValue, NextKeyVal);
+				NewKeyPositions.Add(FilteredKey);
+				KeyHandlesToModify.Add(KeyHandles[KeyIndex + 1]);
+			}
+
+			Curve->Modify();
+			Curve->SetKeyPositions(KeyHandlesToModify, NewKeyPositions);
+		}
+	}
+	OutKeysToSelect = InKeysToOperateOn;
+}
diff --git a/Engine/Source/Editor/CurveEditor/Public/Filters/CurveEditorEulerFilter.h b/Engine/Source/Editor/CurveEditor/Public/Filters/CurveEditorEulerFilter.h
new file mode 100644
--- /dev/null
+++ b/Engine/Source/Editor/CurveEditor/Public/Filters/CurveEditorEulerFilter.h
@@ -0,0 +1,20 @@
+// Copyright 1998-2019 Epic Games, Inc. All Rights Reserved.
+
+#pragma once
+
+#include "Filters/CurveEditorFilterBase.h"
+#include "CurveEditorTypes.h"
+#include "CurveEditorEulerFilter.generated.h"
+
+UCLASS(DisplayName = "Euler")
+class CURVEEDITOR_API UCurveEditorEulerFilter : public UCurveEditorFilterBase
+{
+	GENERATED_BODY()
+public:
+	UCurveEditorEulerFilter(){}
+protected:
+	virtual void ApplyFilter_Impl(TSharedRef<FCurveEditor> InCurveEditor, const TMap<FCurveModelID, FKeyHandleSet>& InKeysToOperateOn, TMap<FCurveModelID, FKeyHandleSet>& OutKeysToSelect) override;
+
+public:
+	
+};
\ No newline at end of file
