Right clicking directly on a key will now show the key-specific context menu instead of the curve menu. If you right click on a curve, not directly over a key you will get that curve-specific menu (regardless of selection). This restores the old behavior but keeps the new behavior of having expanded room to select curves.

diff --git a/Engine/Source/Editor/CurveEditor/Private/CurveEditorContextMenu.cpp b/Engine/Source/Editor/CurveEditor/Private/CurveEditorContextMenu.cpp
--- a/Engine/Source/Editor/CurveEditor/Private/CurveEditorContextMenu.cpp
+++ b/Engine/Source/Editor/CurveEditor/Private/CurveEditorContextMenu.cpp
@@ -19,10 +19,11 @@ void FCurveEditorContextMenu::BuildMenu(FMenuBuilder& MenuBuilder, TSharedRef<FC
 		return FText::Format(LOCTEXT("ApplyStoredCurvesContextMenu", "Apply {0} Stored Curves"), LocalCurveEditor->GetNumBufferedCurves());
 	}));
 
-	// We prioritize key selections (because otherwise you have to be exactly on a key to context menu)
-	// but this creates an issue where you have a selection but try to right click on a curve it still does
-	// the key one... so this lets us get the key menu anywhere on the graph except for over a curve which seems fair.
-	if (NumSelectedKeys > 0 && !HoveredCurveID.IsSet())
+	// We prioritize key selections over curve selections to reduce the pixel-perfectness needed
+	// to edit the keys (which is more common than curves). Right clicking on a key or an empty space
+	// should show the key menu, otherwise we show the curve menu (ie: right clicking on a curve, not 
+	// directly over a key).
+	if (NumSelectedKeys > 0 && (!HoveredCurveID.IsSet() || ClickedPoint.IsSet()))
 	{
 		MenuBuilder.BeginSection("CurveEditorKeySection", FText::Format(LOCTEXT("CurveEditorKeySection", "{0} Selected {0}|plural(one=Key,other=Keys)"), NumSelectedKeys));
 		{
