Curve Editor: Fix marquee select bug in transform tool

diff --git a/Engine/Source/Editor/CurveEditor/Private/SCurveEditorViewContainer.cpp b/Engine/Source/Editor/CurveEditor/Private/SCurveEditorViewContainer.cpp
--- a/Engine/Source/Editor/CurveEditor/Private/SCurveEditorViewContainer.cpp
+++ b/Engine/Source/Editor/CurveEditor/Private/SCurveEditorViewContainer.cpp
@@ -194,8 +194,14 @@ FReply SCurveEditorViewContainer::OnMouseButtonDoubleClick(const FGeometry& InMy
 FReply SCurveEditorViewContainer::OnMouseMove(const FGeometry& MyGeometry, const FPointerEvent& MouseEvent)
 {
 	FVector2D MousePixel = MyGeometry.AbsoluteToLocal(MouseEvent.GetScreenSpacePosition());
+	FReply ToolReply = FReply::Unhandled();
 
-	if (DragOperation.IsSet())
+	if (CurveEditor->GetCurrentTool())
+	{
+		ToolReply = CurveEditor->GetCurrentTool()->OnMouseMove(AsShared(), MyGeometry, MouseEvent);
+	}
+
+	if (!ToolReply.IsEventHandled() && DragOperation.IsSet())
 	{
 		FVector2D InitialPosition = DragOperation->GetInitialPosition();
 
@@ -211,12 +217,7 @@ FReply SCurveEditorViewContainer::OnMouseMove(const FGeometry& MyGeometry, const
 		}
 	}
 
-	if (CurveEditor->GetCurrentTool())
-	{
-		return CurveEditor->GetCurrentTool()->OnMouseMove(AsShared(), MyGeometry, MouseEvent);
-	}
-
-	return FReply::Unhandled();
+	return ToolReply;
 }
 
 FReply SCurveEditorViewContainer::OnMouseButtonUp(const FGeometry& MyGeometry, const FPointerEvent& MouseEvent)
