Make standard curve view API public and expose virtual functions as protected to allow inheriting and extending behavior in plugins.

diff --git a/Engine/Source/Editor/CurveEditor/Private/Views/SCurveEditorViewAbsolute.h b/Engine/Source/Editor/CurveEditor/Private/Views/SCurveEditorViewAbsolute.h
deleted file mode 100644
diff --git a/Engine/Source/Editor/CurveEditor/Private/Views/SCurveEditorViewNormalized.h b/Engine/Source/Editor/CurveEditor/Private/Views/SCurveEditorViewNormalized.h
deleted file mode 100644
diff --git a/Engine/Source/Editor/CurveEditor/Private/Views/SCurveEditorViewStacked.cpp b/Engine/Source/Editor/CurveEditor/Private/Views/SCurveEditorViewStacked.cpp
--- a/Engine/Source/Editor/CurveEditor/Private/Views/SCurveEditorViewStacked.cpp
+++ b/Engine/Source/Editor/CurveEditor/Private/Views/SCurveEditorViewStacked.cpp
@@ -9,8 +9,6 @@
 #include "Algo/Copy.h"
 #include "EditorStyleSet.h"
 
-constexpr float StackedHeight = 150.f;
-constexpr float StackedPadding = 10.f;
 
 void SCurveEditorViewStacked::Construct(const FArguments& InArgs, TWeakPtr<FCurveEditor> InCurveEditor)
 {
@@ -18,6 +16,9 @@ void SCurveEditorViewStacked::Construct(const FArguments& InArgs, TWeakPtr<FCurv
 	OutputMin = 0.0;
 	OutputMax = 1.0;
 
+	StackedHeight = 150.0f;
+	StackedPadding = 10.0f;
+
 	SInteractiveCurveEditorView::Construct(InArgs, InCurveEditor);
 }
 
diff --git a/Engine/Source/Editor/CurveEditor/Private/Views/SCurveEditorViewStacked.h b/Engine/Source/Editor/CurveEditor/Private/Views/SCurveEditorViewStacked.h
deleted file mode 100644
diff --git a/Engine/Source/Editor/CurveEditor/Public/Views/SCurveEditorViewAbsolute.h b/Engine/Source/Editor/CurveEditor/Public/Views/SCurveEditorViewAbsolute.h
new file mode 100644
--- /dev/null
+++ b/Engine/Source/Editor/CurveEditor/Public/Views/SCurveEditorViewAbsolute.h
@@ -0,0 +1,16 @@
+// Copyright 1998-2019 Epic Games, Inc. All Rights Reserved.
+
+#pragma once
+
+#include "Views/SInteractiveCurveEditorView.h"
+
+/**
+ * A Normalized curve view supporting one or more curves with their own screen transform that normalizes the vertical curve range to [-1,1]
+ */
+class CURVEEDITOR_API SCurveEditorViewAbsolute : public SInteractiveCurveEditorView
+{
+public:
+
+	void Construct(const FArguments& InArgs, TWeakPtr<FCurveEditor> InCurveEditor);
+
+};
diff --git a/Engine/Source/Editor/CurveEditor/Public/Views/SCurveEditorViewNormalized.h b/Engine/Source/Editor/CurveEditor/Public/Views/SCurveEditorViewNormalized.h
new file mode 100644
--- /dev/null
+++ b/Engine/Source/Editor/CurveEditor/Public/Views/SCurveEditorViewNormalized.h
@@ -0,0 +1,28 @@
+// Copyright 1998-2019 Epic Games, Inc. All Rights Reserved.
+
+#pragma once
+
+#include "Views/SInteractiveCurveEditorView.h"
+
+/**
+ * A Normalized curve view supporting one or more curves with their own screen transform that normalizes the vertical curve range to [-1,1]
+ */
+class CURVEEDITOR_API SCurveEditorViewNormalized : public SInteractiveCurveEditorView
+{
+public:
+
+	void Construct(const FArguments& InArgs, TWeakPtr<FCurveEditor> InCurveEditor);
+
+	/** Tools should use vertical snapping since grid lines to snap to will usually be visible */
+	virtual bool IsValueSnapEnabled() const override { return true; }
+
+protected:
+
+	virtual void Tick(const FGeometry& AllottedGeometry, const double InCurrentTime, const float InDeltaTime) override;
+
+	virtual void PaintView(const FPaintArgs& Args, const FGeometry& AllottedGeometry, const FSlateRect& MyCullingRect, FSlateWindowElementList& OutDrawElements, int32 BaseLayerId, const FWidgetStyle& InWidgetStyle, bool bParentEnabled) const override;
+
+	virtual void GetGridLinesY(TSharedRef<const FCurveEditor> CurveEditor, TArray<float>& MajorGridLines, TArray<float>& MinorGridLines, TArray<FText>* MajorGridLabels) const override;
+
+	virtual void DrawBufferedCurves(const FGeometry& AllottedGeometry, const FSlateRect& MyCullingRect, FSlateWindowElementList& OutDrawElements, int32 BaseLayerId, ESlateDrawEffect DrawEffects) const;
+};
diff --git a/Engine/Source/Editor/CurveEditor/Public/Views/SCurveEditorViewStacked.h b/Engine/Source/Editor/CurveEditor/Public/Views/SCurveEditorViewStacked.h
new file mode 100644
--- /dev/null
+++ b/Engine/Source/Editor/CurveEditor/Public/Views/SCurveEditorViewStacked.h
@@ -0,0 +1,36 @@
+// Copyright 1998-2019 Epic Games, Inc. All Rights Reserved.
+
+#pragma once
+
+#include "Views/SInteractiveCurveEditorView.h"
+
+/**
+ * A Normalized curve view supporting one or more curves with their own screen transform that normalizes the vertical curve range to [-1,1]
+ */
+class CURVEEDITOR_API SCurveEditorViewStacked : public SInteractiveCurveEditorView
+{
+public:
+
+	void Construct(const FArguments& InArgs, TWeakPtr<FCurveEditor> InCurveEditor);
+	
+	/** Tools should use vertical snapping since grid lines to snap to will usually be visible */
+	virtual bool IsValueSnapEnabled() const override { return true; }
+
+protected:
+
+	virtual void Tick(const FGeometry& AllottedGeometry, const double InCurrentTime, const float InDeltaTime) override;
+	virtual FVector2D ComputeDesiredSize(float LayoutScaleMultiplier) const override;
+
+	virtual void PaintView(const FPaintArgs& Args, const FGeometry& AllottedGeometry, const FSlateRect& MyCullingRect, FSlateWindowElementList& OutDrawElements, int32 BaseLayerId, const FWidgetStyle& InWidgetStyle, bool bParentEnabled) const override;
+	virtual void GetGridLinesY(TSharedRef<const FCurveEditor> CurveEditor, TArray<float>& MajorGridLines, TArray<float>& MinorGridLines, TArray<FText>* MajorGridLabels = nullptr) const override;
+
+	virtual void DrawViewGrids(const FGeometry& AllottedGeometry, const FSlateRect& MyCullingRect, FSlateWindowElementList& OutDrawElements, int32 BaseLayerId, ESlateDrawEffect DrawEffects) const;
+	virtual void DrawLabels(const FGeometry& AllottedGeometry, const FSlateRect& MyCullingRect, FSlateWindowElementList& OutDrawElements, int32 BaseLayerId, ESlateDrawEffect DrawEffects) const;
+	virtual void DrawBufferedCurves(const FGeometry& AllottedGeometry, const FSlateRect& MyCullingRect, FSlateWindowElementList& OutDrawElements, int32 BaseLayerId, ESlateDrawEffect DrawEffects) const;
+
+	/** Stacked height per curve */
+	float StackedHeight;
+
+	/** Stacked padding per curve */
+	float StackedPadding;
+};
