Add read only disablement to missing locations & ability to disable key drawing to CurveEditor

diff --git a/Engine/Source/Editor/CurveEditor/Private/RichCurveEditorModel.cpp b/Engine/Source/Editor/CurveEditor/Private/RichCurveEditorModel.cpp
--- a/Engine/Source/Editor/CurveEditor/Private/RichCurveEditorModel.cpp
+++ b/Engine/Source/Editor/CurveEditor/Private/RichCurveEditorModel.cpp
@@ -276,6 +276,11 @@ void FRichCurveEditorModel::GetKeyPositions(TArrayView<const FKeyHandle> InKeys,
 
 void FRichCurveEditorModel::SetKeyPositions(TArrayView<const FKeyHandle> InKeys, TArrayView<const FKeyPosition> InKeyPositions)
 {
+	if (IsReadOnly())
+	{
+		return;
+	}
+
 	if (UObject* Owner = WeakOwner.Get())
 	{
 		Owner->Modify();
@@ -345,6 +350,11 @@ void FRichCurveEditorModel::GetKeyAttributes(TArrayView<const FKeyHandle> InKeys
 
 void FRichCurveEditorModel::SetKeyAttributes(TArrayView<const FKeyHandle> InKeys, TArrayView<const FKeyAttributes> InAttributes)
 {
+	if (IsReadOnly())
+	{
+		return;
+	}
+
 	if (UObject* Owner = WeakOwner.Get())
 	{
 		const TArray<FRichCurveKey>& AllKeys = RichCurve->GetConstRefOfKeys();
diff --git a/Engine/Source/Editor/CurveEditor/Private/Views/SInteractiveCurveEditorView.cpp b/Engine/Source/Editor/CurveEditor/Private/Views/SInteractiveCurveEditorView.cpp
--- a/Engine/Source/Editor/CurveEditor/Private/Views/SInteractiveCurveEditorView.cpp
+++ b/Engine/Source/Editor/CurveEditor/Private/Views/SInteractiveCurveEditorView.cpp
@@ -415,30 +415,33 @@ void SInteractiveCurveEditorView::DrawCurves(TSharedRef<FCurveEditor> CurveEdito
 		LinePoints.SetNum(2);
 
 		// Draw tangents
-		for(int32 PointIndex = 0; PointIndex < Params.Points.Num(); PointIndex++)
+		if (Params.bKeyDrawEnabled)
 		{
-			const FCurvePointInfo& Point = Params.Points[PointIndex];
-			const FKeyDrawInfo& PointDrawInfo = Params.GetKeyDrawInfo(Point.Type, PointIndex);
-			const bool          bSelected = CurveEditor->GetSelection().IsSelected(FCurvePointHandle(Params.GetID(), Point.Type, Point.KeyHandle));
-			const FLinearColor  PointTint = bSelected ? SelectionColor : PointDrawInfo.Tint;
+			for (int32 PointIndex = 0; PointIndex < Params.Points.Num(); PointIndex++)
+			{
+				const FCurvePointInfo& Point = Params.Points[PointIndex];
+				const FKeyDrawInfo& PointDrawInfo = Params.GetKeyDrawInfo(Point.Type, PointIndex);
+				const bool          bSelected = CurveEditor->GetSelection().IsSelected(FCurvePointHandle(Params.GetID(), Point.Type, Point.KeyHandle));
+				const FLinearColor  PointTint = bSelected ? SelectionColor : PointDrawInfo.Tint;
 
-			const int32 KeyLayerId = BaseLayerId + Point.LayerBias + (bSelected ? CurveViewConstants::ELayerOffset::SelectedKeys : CurveViewConstants::ELayerOffset::Keys);
+				const int32 KeyLayerId = BaseLayerId + Point.LayerBias + (bSelected ? CurveViewConstants::ELayerOffset::SelectedKeys : CurveViewConstants::ELayerOffset::Keys);
 
-			if (Point.LineDelta.X != 0.f || Point.LineDelta.Y != 0.f)
-			{
-				LinePoints[0] = Point.ScreenPosition + Point.LineDelta.GetSafeNormal() * (PointDrawInfo.ScreenSize.X*.5f);
-				LinePoints[1] = Point.ScreenPosition + Point.LineDelta;
+				if (Point.LineDelta.X != 0.f || Point.LineDelta.Y != 0.f)
+				{
+					LinePoints[0] = Point.ScreenPosition + Point.LineDelta.GetSafeNormal() * (PointDrawInfo.ScreenSize.X*.5f);
+					LinePoints[1] = Point.ScreenPosition + Point.LineDelta;
 
-				// Draw the connecting line - connecting lines are always drawn below everything else
-				FSlateDrawElement::MakeLines(OutDrawElements, BaseLayerId + CurveViewConstants::ELayerOffset::Keys - 1, PaintGeometry, LinePoints, DrawEffects, PointTint, true);
-			}
+					// Draw the connecting line - connecting lines are always drawn below everything else
+					FSlateDrawElement::MakeLines(OutDrawElements, BaseLayerId + CurveViewConstants::ELayerOffset::Keys - 1, PaintGeometry, LinePoints, DrawEffects, PointTint, true);
+				}
 
-			FPaintGeometry PointGeometry = AllottedGeometry.ToPaintGeometry(
-				Point.ScreenPosition - (PointDrawInfo.ScreenSize * 0.5f),
-				PointDrawInfo.ScreenSize
-			);
+				FPaintGeometry PointGeometry = AllottedGeometry.ToPaintGeometry(
+					Point.ScreenPosition - (PointDrawInfo.ScreenSize * 0.5f),
+					PointDrawInfo.ScreenSize
+				);
 
-			FSlateDrawElement::MakeBox(OutDrawElements, KeyLayerId, PointGeometry, PointDrawInfo.Brush, DrawEffects, PointTint);
+				FSlateDrawElement::MakeBox(OutDrawElements, KeyLayerId, PointGeometry, PointDrawInfo.Brush, DrawEffects, PointTint);
+			}
 		}
 	}
 }
@@ -521,6 +524,7 @@ void SInteractiveCurveEditorView::GetCurveDrawParams(TArray<FCurveDrawParams>& O
 		// Create a new set of Curve Drawing Parameters to represent this particular Curve
 		FCurveDrawParams Params(Pair.Key);
 		Params.Color = CurveModel->GetColor();
+		Params.bKeyDrawEnabled = CurveModel->IsKeyDrawEnabled();
 
 		// Gather the display metrics to use for each key type. This allows a Curve Model to override
 		// whether or not the curve supports Keys, Arrive/Leave Tangents, etc. If the Curve Model doesn't
@@ -922,7 +926,7 @@ FReply SInteractiveCurveEditorView::OnMouseButtonDown(const FGeometry& MyGeometr
 			if (TOptional<FCurveModelID> HoveredCurve = GetHoveredCurve())
 			{
 				FCurveModel* CurveToAddTo = CurveEditor->FindCurve(HoveredCurve.GetValue());
-				if (CurveToAddTo)
+				if (CurveToAddTo && !CurveToAddTo->IsReadOnly())
 				{
 					FScopedTransaction Transaction(LOCTEXT("InsertKey", "Insert Key"));
 
@@ -1000,35 +1004,41 @@ FReply SInteractiveCurveEditorView::OnMouseButtonDown(const FGeometry& MyGeometr
 		TOptional<FCurvePointHandle> MouseDownPoint = HitPoint(MousePixel);
 		if (MouseDownPoint.IsSet())
 		{
-			if (bShiftPressed)
+			if (FCurveModel* CurveModel = CurveEditor->FindCurve(MouseDownPoint->CurveID))
 			{
-				CurveEditor->GetSelection().Add(MouseDownPoint.GetValue());
-			}
-			else if (bCtrlPressed)
-			{
-				CurveEditor->GetSelection().Toggle(MouseDownPoint.GetValue());
-			}
-			else
-			{
-				if (CurveEditor->GetSelection().Contains(MouseDownPoint->CurveID, MouseDownPoint->KeyHandle))
-				{
-					CurveEditor->GetSelection().ChangeSelectionPointType(MouseDownPoint->PointType);
-				}
-				else
+				if (!CurveModel->IsReadOnly())
 				{
-					CurveEditor->GetSelection().Clear();
-					CurveEditor->GetSelection().Add(MouseDownPoint.GetValue());
-				}
-			}
+					if (bShiftPressed)
+					{
+						CurveEditor->GetSelection().Add(MouseDownPoint.GetValue());
+					}
+					else if (bCtrlPressed)
+					{
+						CurveEditor->GetSelection().Toggle(MouseDownPoint.GetValue());
+					}
+					else
+					{
+						if (CurveEditor->GetSelection().Contains(MouseDownPoint->CurveID, MouseDownPoint->KeyHandle))
+						{
+							CurveEditor->GetSelection().ChangeSelectionPointType(MouseDownPoint->PointType);
+						}
+						else
+						{
+							CurveEditor->GetSelection().Clear();
+							CurveEditor->GetSelection().Add(MouseDownPoint.GetValue());
+						}
+					}
 
-			TUniquePtr<ICurveEditorKeyDragOperation> KeyDrag = CreateKeyDrag(MouseDownPoint->PointType);
+					TUniquePtr<ICurveEditorKeyDragOperation> KeyDrag = CreateKeyDrag(MouseDownPoint->PointType);
 
-			KeyDrag->Initialize(CurveEditor.Get(), MouseDownPoint);
+					KeyDrag->Initialize(CurveEditor.Get(), MouseDownPoint);
 
-			DragOperation = FCurveEditorDelayedDrag(MousePixel, MouseEvent.GetEffectingButton());
-			DragOperation->DragImpl = MoveTemp(KeyDrag);
+					DragOperation = FCurveEditorDelayedDrag(MousePixel, MouseEvent.GetEffectingButton());
+					DragOperation->DragImpl = MoveTemp(KeyDrag);
 
-			return FReply::Handled();
+					return FReply::Handled();
+				}
+			}
 		}
 	}
 
@@ -1363,6 +1373,11 @@ void SInteractiveCurveEditorView::AddKeyAtTime(const TSet<FCurveModelID>& ToCurv
 		FCurveModel* CurveModel = CurveEditor->FindCurve(CurveModelID);
 		check(CurveModel);
 
+		if (CurveModel->IsReadOnly())
+		{
+			continue;
+		}
+
 		// Ensure the time is snapped if needed
 		FCurveSnapMetrics SnapMetrics = CurveEditor->GetCurveSnapMetrics(CurveModelID);
 		double SnappedTime = SnapMetrics.SnapInputSeconds(InTime);
diff --git a/Engine/Source/Editor/CurveEditor/Public/CurveDrawInfo.h b/Engine/Source/Editor/CurveEditor/Public/CurveDrawInfo.h
--- a/Engine/Source/Editor/CurveEditor/Public/CurveDrawInfo.h
+++ b/Engine/Source/Editor/CurveEditor/Public/CurveDrawInfo.h
@@ -81,11 +81,15 @@ struct FCurveDrawParams
 	/** Value defining how to draw keys of type ECurvePointType::LeaveTangent. */
 	FKeyDrawInfo LeaveTangentDrawInfo;
 
+	/** Whether or not to draw keys on the curve. */
+	uint8 bKeyDrawEnabled : 1;
+
 	/**
 	 * Construct new draw parameters for the specified curve ID
 	 */
 	FCurveDrawParams(FCurveModelID InID)
 		: Color(FLinearColor::White)
+		, bKeyDrawEnabled(1)
 		, ID(InID)
 	{}
 
diff --git a/Engine/Source/Editor/CurveEditor/Public/CurveModel.h b/Engine/Source/Editor/CurveEditor/Public/CurveModel.h
--- a/Engine/Source/Editor/CurveEditor/Public/CurveModel.h
+++ b/Engine/Source/Editor/CurveEditor/Public/CurveModel.h
@@ -38,6 +38,7 @@ public:
 
 	FCurveModel()
 		: Color(FLinearColor::White)
+		, bKeyDrawEnabled(1)
 		, SupportedViews(ECurveEditorViewID::ANY_BUILT_IN)
 	{}
 
@@ -195,6 +196,14 @@ public:
 		return 0.0;
 	}
 
+	/**
+	 * Retrieve this curve's color
+	 */
+	virtual FLinearColor GetColor() const
+	{
+		return IsReadOnly() ? Color.Desaturate(.6f) : Color;
+	}
+
 	/**
 	 * Create key proxy objects for the specified key handles. One object should be assigned to OutObjects per index within InKeyHandles
 	 *
@@ -305,19 +314,26 @@ public:
 	}
 
 	/**
-	 * Retrieve this curve's color
 	 */
-	FORCEINLINE FLinearColor GetColor() const
+	FORCEINLINE void SetColor(const FLinearColor& InColor)
 	{
-		return IsReadOnly() ? Color.Desaturate(.6f) : Color;
+		Color = InColor;
 	}
 
 	/**
-	 * Assign a new color to this curve
+ * Retrieves whether or not to disable drawing keys
+ */
+	FORCEINLINE bool IsKeyDrawEnabled() const
+	{
+		return bKeyDrawEnabled != 0;
+	}
+
+	/**
+	 * Assign whether or not to disable drawing keys
 	 */
-	FORCEINLINE void SetColor(const FLinearColor& InColor)
+	FORCEINLINE void SetIsKeyDrawEnabled(bool bInKeyDrawEnabled)
 	{
-		Color = InColor;
+		bKeyDrawEnabled = bInKeyDrawEnabled ? 1 : 0;
 	}
 
 	/**
@@ -342,6 +358,9 @@ protected:
 	/** This curve's display color */
 	FLinearColor Color;
 
+	/** Whether or not to draw curve's keys */
+	uint8 bKeyDrawEnabled : 1;
+
 	/** A set of views supported by this curve */
 	ECurveEditorViewID SupportedViews;
 };
