Curve Editor: Show a dummy editable text box before the property rows are generated.

diff --git a/Engine/Source/Editor/CurveEditor/Private/SCurveKeyDetailPanel.cpp b/Engine/Source/Editor/CurveEditor/Private/SCurveKeyDetailPanel.cpp
--- a/Engine/Source/Editor/CurveEditor/Private/SCurveKeyDetailPanel.cpp
+++ b/Engine/Source/Editor/CurveEditor/Private/SCurveKeyDetailPanel.cpp
@@ -20,8 +20,63 @@ void SCurveKeyDetailPanel::Construct(const FArguments& InArgs, TSharedRef<FCurve
 
 	PropertyRowGenerator = PropertyEditorModule.CreatePropertyRowGenerator(Args);
 	PropertyRowGenerator->OnRowsRefreshed().AddSP(this, &SCurveKeyDetailPanel::PropertyRowsRefreshed);
+
+	PropertyRowsRefreshed();
 }
 
+// A Dummy editable text box that is visible before property rows are generated.
+class STempConstrainedBox : public SCompoundWidget
+{
+public:
+	SLATE_BEGIN_ARGS(STempConstrainedBox)
+		: _MinWidth(125.f)
+		, _MaxWidth(125.f)
+	{}
+	SLATE_DEFAULT_SLOT(FArguments, Content)
+		SLATE_ATTRIBUTE(TOptional<float>, MinWidth)
+		SLATE_ATTRIBUTE(TOptional<float>, MaxWidth)
+		SLATE_END_ARGS()
+
+	void Construct(const FArguments& InArgs)
+	{
+		MinWidth = InArgs._MinWidth;
+		MaxWidth = InArgs._MaxWidth;
+
+		ChildSlot
+			[
+				SNew(SEditableTextBox)
+			];
+	};
+
+	virtual FVector2D ComputeDesiredSize(float LayoutScaleMultiplier) const override
+	{
+		const float MinWidthVal = MinWidth.Get().Get(0.0f);
+		const float MaxWidthVal = MaxWidth.Get().Get(0.0f);
+
+		if (MinWidthVal == 0.0f && MaxWidthVal == 0.0f)
+		{
+			return SCompoundWidget::ComputeDesiredSize(LayoutScaleMultiplier);
+		}
+		else
+		{
+			FVector2D ChildSize = ChildSlot.GetWidget()->GetDesiredSize();
+
+			float XVal = FMath::Max(MinWidthVal, ChildSize.X);
+			if (MaxWidthVal >= MinWidthVal)
+			{
+				XVal = FMath::Min(MaxWidthVal, XVal);
+			}
+
+			return FVector2D(XVal, ChildSize.Y);
+		}
+	}
+
+private:
+	TAttribute< TOptional<float> > MinWidth;
+	TAttribute< TOptional<float> > MaxWidth;
+};
+
+
 void SCurveKeyDetailPanel::PropertyRowsRefreshed()
 {
 	// UE_LOG(LogTemp, Log, TEXT("PropertyRowsRefreshed"));
@@ -55,6 +110,26 @@ void SCurveKeyDetailPanel::PropertyRowsRefreshed()
 		}
 	}
 
+	if (!TimeWidget)
+	{
+		if (!TempTimeWidget.IsValid())
+		{
+			TempTimeWidget = SNew(STempConstrainedBox);
+		}
+
+		TimeWidget = TempTimeWidget;
+	}
+
+	if (!ValueWidget)
+	{
+		if (!TempValueWidget.IsValid())
+		{
+			TempValueWidget = SNew(STempConstrainedBox);
+		}
+
+		ValueWidget = TempValueWidget;
+	}
+
 	if (TimeWidget && ValueWidget)
 	{
 		ConstructChildLayout(TimeWidget, ValueWidget);
diff --git a/Engine/Source/Editor/CurveEditor/Public/SCurveKeyDetailPanel.h b/Engine/Source/Editor/CurveEditor/Public/SCurveKeyDetailPanel.h
--- a/Engine/Source/Editor/CurveEditor/Public/SCurveKeyDetailPanel.h
+++ b/Engine/Source/Editor/CurveEditor/Public/SCurveKeyDetailPanel.h
@@ -29,4 +29,7 @@ private:
 
 private:
 	TSharedPtr<IPropertyRowGenerator> PropertyRowGenerator;
+
+	TSharedPtr<SWidget> TempTimeWidget;
+	TSharedPtr<SWidget> TempValueWidget;
 };
\ No newline at end of file
