Curve Editor: Prevent throttling on drag so that editing keys updates the viewport immediately.

diff --git a/Engine/Plugins/Editor/CurveEditorTools/Source/CurveEditorTools/Private/Tools/CurveEditorMultiScaleTool.cpp b/Engine/Plugins/Editor/CurveEditorTools/Source/CurveEditorTools/Private/Tools/CurveEditorMultiScaleTool.cpp
--- a/Engine/Plugins/Editor/CurveEditorTools/Source/CurveEditorTools/Private/Tools/CurveEditorMultiScaleTool.cpp
+++ b/Engine/Plugins/Editor/CurveEditorTools/Source/CurveEditorTools/Private/Tools/CurveEditorMultiScaleTool.cpp
@@ -242,7 +242,7 @@ FReply FCurveEditorMultiScaleTool::OnMouseButtonDown(TSharedRef<SWidget> OwningW
 			// Start a Delayed Drag
 			DelayedDrag = FDelayedDrag(MyGeometry.AbsoluteToLocal(MouseEvent.GetScreenSpacePosition()), MouseEvent.GetEffectingButton());
 
-			return FReply::Handled();
+			return FReply::Handled().PreventThrottling();
 		}
 	}
 	return FReply::Unhandled();
diff --git a/Engine/Plugins/Editor/CurveEditorTools/Source/CurveEditorTools/Private/Tools/CurveEditorRetimeTool.cpp b/Engine/Plugins/Editor/CurveEditorTools/Source/CurveEditorTools/Private/Tools/CurveEditorRetimeTool.cpp
--- a/Engine/Plugins/Editor/CurveEditorTools/Source/CurveEditorTools/Private/Tools/CurveEditorRetimeTool.cpp
+++ b/Engine/Plugins/Editor/CurveEditorTools/Source/CurveEditorTools/Private/Tools/CurveEditorRetimeTool.cpp
@@ -138,7 +138,7 @@ FReply FCurveEditorRetimeTool::OnMouseButtonDown(TSharedRef<SWidget> OwningWidge
 		}
 
 		// We always handle the left mouse button so clicks don't bubble through to selection
-		return FReply::Handled();
+		return FReply::Handled().PreventThrottling();
 	}
 
 	// Right Clicks need to go through so that panning the graph works
diff --git a/Engine/Plugins/Editor/CurveEditorTools/Source/CurveEditorTools/Private/Tools/CurveEditorTransformTool.cpp b/Engine/Plugins/Editor/CurveEditorTools/Source/CurveEditorTools/Private/Tools/CurveEditorTransformTool.cpp
--- a/Engine/Plugins/Editor/CurveEditorTools/Source/CurveEditorTools/Private/Tools/CurveEditorTransformTool.cpp
+++ b/Engine/Plugins/Editor/CurveEditorTools/Source/CurveEditorTools/Private/Tools/CurveEditorTransformTool.cpp
@@ -351,7 +351,7 @@ FReply FCurveEditorTransformTool::OnMouseButtonDown(TSharedRef<SWidget> OwningWi
 			// Start a Delayed Drag
 			DelayedDrag = FDelayedDrag(MyGeometry.AbsoluteToLocal(MouseEvent.GetScreenSpacePosition()), MouseEvent.GetEffectingButton());
 			
-			return FReply::Handled();
+			return FReply::Handled().PreventThrottling();
 		}
 	}
 	return FReply::Unhandled();
diff --git a/Engine/Source/Editor/CurveEditor/Private/Views/SInteractiveCurveEditorView.cpp b/Engine/Source/Editor/CurveEditor/Private/Views/SInteractiveCurveEditorView.cpp
--- a/Engine/Source/Editor/CurveEditor/Private/Views/SInteractiveCurveEditorView.cpp
+++ b/Engine/Source/Editor/CurveEditor/Private/Views/SInteractiveCurveEditorView.cpp
@@ -974,7 +974,7 @@ FReply SInteractiveCurveEditorView::OnMouseButtonDown(const FGeometry& MyGeometr
 			DragOperation = FCurveEditorDelayedDrag(MousePixel, MouseEvent.GetEffectingButton());
 			DragOperation->DragImpl = MoveTemp(KeyDrag);
 
-			return FReply::Handled();
+			return FReply::Handled().PreventThrottling();
 		}
 	}
 	else if (MouseEvent.GetEffectingButton() == EKeys::RightMouseButton)
@@ -1036,7 +1036,7 @@ FReply SInteractiveCurveEditorView::OnMouseButtonDown(const FGeometry& MyGeometr
 					DragOperation = FCurveEditorDelayedDrag(MousePixel, MouseEvent.GetEffectingButton());
 					DragOperation->DragImpl = MoveTemp(KeyDrag);
 
-					return FReply::Handled();
+					return FReply::Handled().PreventThrottling();
 				}
 			}
 		}
