Added a PostEditChangeProperty call to the updates for the individual R,G,B,A keys.
#rnx

diff --git a/Engine/Source/Editor/CurveEditor/Private/RichCurveEditorModel.cpp b/Engine/Source/Editor/CurveEditor/Private/RichCurveEditorModel.cpp
--- a/Engine/Source/Editor/CurveEditor/Private/RichCurveEditorModel.cpp
+++ b/Engine/Source/Editor/CurveEditor/Private/RichCurveEditorModel.cpp
@@ -296,6 +296,8 @@ void FRichCurveEditorModel::SetKeyPositions(TArrayView<const FKeyHandle> InKeys,
 			}
 		}
 		RichCurve->AutoSetTangents();
+		FPropertyChangedEvent PropertyChangeStruct(nullptr, EPropertyChangeType::ValueSet);
+		Owner->PostEditChangeProperty(PropertyChangeStruct);
 	}
 }
 
@@ -485,6 +487,9 @@ void FRichCurveEditorModel::SetKeyAttributes(TArrayView<const FKeyHandle> InKeys
 		{
 			RichCurve->AutoSetTangents();
 		}
+
+		FPropertyChangedEvent PropertyChangeStruct(nullptr, EPropertyChangeType::ValueSet);
+		Owner->PostEditChangeProperty(PropertyChangeStruct);
 	}
 }
 
@@ -512,6 +517,9 @@ void FRichCurveEditorModel::SetCurveAttributes(const FCurveAttributes& InCurveAt
 		{
 			RichCurve->PostInfinityExtrap = InCurveAttributes.GetPostExtrapolation();
 		}
+
+		FPropertyChangedEvent PropertyChangeStruct(nullptr, EPropertyChangeType::ValueSet);
+		Owner->PostEditChangeProperty(PropertyChangeStruct);
 	}
 }
 
