Curve Editor: Add snap to grid option and allow users to set grid spacing

diff --git a/Engine/Plugins/Editor/CurveEditorTools/Source/CurveEditorTools/Private/Tools/CurveEditorRetimeTool.cpp b/Engine/Plugins/Editor/CurveEditorTools/Source/CurveEditorTools/Private/Tools/CurveEditorRetimeTool.cpp
--- a/Engine/Plugins/Editor/CurveEditorTools/Source/CurveEditorTools/Private/Tools/CurveEditorRetimeTool.cpp
+++ b/Engine/Plugins/Editor/CurveEditorTools/Source/CurveEditorTools/Private/Tools/CurveEditorRetimeTool.cpp
@@ -591,7 +591,7 @@ void FCurveEditorRetimeTool::OnDrag(const FVector2D& InStartPosition, const FVec
 
 				// Snap the new values to the grid.
 				FKeyPosition NewPosition = Position;
-				NewPosition.InputValue = View->IsTimeSnapEnabled() ? CurveEditor->GetSnapMetrics().SnapInputSeconds(NewKeyValue) : NewKeyValue;
+				NewPosition.InputValue = View->IsTimeSnapEnabled() ? CurveEditor->GetCurveSnapMetrics(ChannelData.CurveID).SnapInputSeconds(NewKeyValue) : NewKeyValue;
 
 				// Update our Curve Model with our updated data.
 				Curve->SetKeyPositions(MakeArrayView(&ChannelData.Handles[KeyIndex], 1), MakeArrayView(&NewPosition, 1));
diff --git a/Engine/Source/Editor/CurveEditor/Private/CurveEditor.cpp b/Engine/Source/Editor/CurveEditor/Private/CurveEditor.cpp
--- a/Engine/Source/Editor/CurveEditor/Private/CurveEditor.cpp
+++ b/Engine/Source/Editor/CurveEditor/Private/CurveEditor.cpp
@@ -21,6 +21,7 @@
 #include "ITimeSlider.h"
 #include "Framework/Notifications/NotificationManager.h"
 #include "Widgets/Notifications/SNotificationList.h"
+#include "Runtime/Core/Public/Algo/Transform.h"
 
 #define LOCTEXT_NAMESPACE "CurveEditor"
 
@@ -43,7 +44,6 @@ FCurveEditor::FCurveEditor()
 
 	OutputSnapEnabledAttribute = true;
 	InputSnapEnabledAttribute  = true;
-	OutputSnapIntervalAttribute = 0.1;
 	InputSnapRateAttribute = FFrameRate(10, 1);
 
 	GridLineLabelFormatXAttribute = LOCTEXT("GridXLabelFormat", "{0}s");
@@ -306,16 +306,33 @@ void FCurveEditor::BindCommands()
 	}
 }
 
-FCurveEditorSnapMetrics FCurveEditor::GetSnapMetrics() const
+FCurveSnapMetrics FCurveEditor::GetCurveSnapMetrics(FCurveModelID CurveModel) const
 {
-	FCurveEditorSnapMetrics Metrics;
+	FCurveSnapMetrics CurveMetrics;
 
-	Metrics.bSnapOutputValues  = OutputSnapEnabledAttribute.Get();
-	Metrics.OutputSnapInterval = OutputSnapIntervalAttribute.Get();
-	Metrics.bSnapInputValues   = InputSnapEnabledAttribute.Get();
-	Metrics.InputSnapRate      = InputSnapRateAttribute.Get();
+	const SCurveEditorView* View = FindFirstInteractiveView(CurveModel);
+	if (!View)
+	{
+		return CurveMetrics;
+	}
 
-	return Metrics;
+	// get the grid lines in view space
+	TArray<float> ViewSpaceGridLines;
+	View->GetGridLinesY(SharedThis(this), ViewSpaceGridLines, ViewSpaceGridLines);
+
+	// convert the grid lines from view space
+	TArray<double> CurveSpaceGridLines;
+	ViewSpaceGridLines.Reserve(ViewSpaceGridLines.Num());
+	FCurveEditorScreenSpace CurveSpace = View->GetCurveSpace(CurveModel);
+	Algo::Transform(ViewSpaceGridLines, CurveSpaceGridLines, [&CurveSpace](float VSVal) { return CurveSpace.ScreenToValue(VSVal); });
+	
+	// create metrics struct;
+	CurveMetrics.bSnapOutputValues = OutputSnapEnabledAttribute.Get();
+	CurveMetrics.bSnapInputValues = InputSnapEnabledAttribute.Get();
+	CurveMetrics.AllGridLines = CurveSpaceGridLines;
+	CurveMetrics.InputSnapRate = InputSnapRateAttribute.Get();
+
+	return CurveMetrics;
 }
 
 void FCurveEditor::ZoomToFit(EAxisList::Type Axes)
@@ -422,8 +439,6 @@ void FCurveEditor::ZoomToFitInternal(EAxisList::Type Axes, const TMap<FCurveMode
 		}
 	}
 
-	FCurveEditorSnapMetrics SnapMetrics = GetSnapMetrics();
-
 	if (Axes & EAxisList::X && InputMin != TNumericLimits<double>::Max() && InputMax != TNumericLimits<double>::Lowest())
 	{
 		// If zooming to the same (or invalid) min/max, keep the same zoom scale and center within the timeline
@@ -438,7 +453,7 @@ void FCurveEditor::ZoomToFitInternal(EAxisList::Type Axes, const TMap<FCurveMode
 		}
 		else
 		{
-			const double MinInputZoom = SnapMetrics.bSnapInputValues ? SnapMetrics.InputSnapRate.AsInterval() : 0.00001;
+			const double MinInputZoom = InputSnapEnabledAttribute.Get() ? InputSnapRateAttribute.Get().AsInterval() : 0.00001;
 			const double InputPadding = FMath::Max((InputMax - InputMin) * 0.1, MinInputZoom);
 			InputMax = FMath::Max(InputMin + MinInputZoom, InputMax);
 
@@ -465,7 +480,7 @@ void FCurveEditor::ZoomToFitInternal(EAxisList::Type Axes, const TMap<FCurveMode
 		}
 		else
 		{
-			const double MinOutputZoom = SnapMetrics.bSnapOutputValues ? SnapMetrics.OutputSnapInterval : 0.00001;
+			constexpr double MinOutputZoom = 0.00001;
 			const double OutputPadding = FMath::Max((OutputMax - OutputMin) * 0.05, MinOutputZoom);
 
 			
@@ -664,13 +679,13 @@ FCurveEditorScreenSpaceH FCurveEditor::GetPanelInputSpace() const
 	return FCurveEditorScreenSpaceH(PanelWidth, InputMin, InputMax);
 }
 
-void FCurveEditor::ConstructXGridLines(TArray<float>& MajorGridLines, TArray<float>& MinorGridLines, TArray<FText>& MajorGridLabels) const
+void FCurveEditor::ConstructXGridLines(TArray<float>& MajorGridLines, TArray<float>& MinorGridLines, TArray<FText>* MajorGridLabels) const
 {
 	FCurveEditorScreenSpaceH InputSpace = GetPanelInputSpace();
 
 	double MajorGridStep  = 0.0;
 	int32  MinorDivisions = 0;
-	if (GetSnapMetrics().InputSnapRate.ComputeGridSpacing(InputSpace.PixelsPerInput(), MajorGridStep, MinorDivisions))
+	if (InputSnapRateAttribute.Get().ComputeGridSpacing(InputSpace.PixelsPerInput(), MajorGridStep, MinorDivisions))
 	{
 		FText GridLineLabelFormatX = GridLineLabelFormatXAttribute.Get();
 		const double FirstMajorLine = FMath::FloorToDouble(InputSpace.GetInputMin() / MajorGridStep) * MajorGridStep;
@@ -679,7 +694,10 @@ void FCurveEditor::ConstructXGridLines(TArray<float>& MajorGridLines, TArray<flo
 		for (double CurrentMajorLine = FirstMajorLine; CurrentMajorLine < LastMajorLine; CurrentMajorLine += MajorGridStep)
 		{
 			MajorGridLines.Add( (CurrentMajorLine - InputSpace.GetInputMin()) * InputSpace.PixelsPerInput() );
-			MajorGridLabels.Add( FText::Format(GridLineLabelFormatX, FText::AsNumber(CurrentMajorLine)) );
+			if (MajorGridLabels)
+			{
+				MajorGridLabels->Add(FText::Format(GridLineLabelFormatX, FText::AsNumber(CurrentMajorLine)));
+			}
 
 			for (int32 Step = 1; Step < MinorDivisions; ++Step)
 			{
diff --git a/Engine/Source/Editor/CurveEditor/Private/CurveEditorHelpers.cpp b/Engine/Source/Editor/CurveEditor/Private/CurveEditorHelpers.cpp
--- a/Engine/Source/Editor/CurveEditor/Private/CurveEditorHelpers.cpp
+++ b/Engine/Source/Editor/CurveEditor/Private/CurveEditorHelpers.cpp
@@ -1,6 +1,8 @@
 // Copyright 1998-2019 Epic Games, Inc. All Rights Reserved.
 
 #include "CurveEditorHelpers.h"
+#include "Fonts/FontMeasure.h"
+#include "Slate/Public/Framework/Application/SlateApplication.h"
 
 
 namespace CurveEditor
@@ -76,4 +78,45 @@ void ConstructYGridLines(const FCurveEditorScreenSpace& ViewSpace, uint8 InMinor
 	}
 }
 
+void ConstructFixedYGridLines(const FCurveEditorScreenSpace& ViewSpace, uint8 InMinorDivisions, double InMinorGridStep, TArray<float>& OutMajorGridLines, TArray<float>& OutMinorGridLines, FText GridLineLabelFormatY, 
+	TArray<FText>* OutMajorGridLabels, TOptional<double> InOutputMin, TOptional<double> InOutputMax)
+{
+	const double MajorGridStep = InMinorGridStep * InMinorDivisions;
+	const double FirstMinorLine = InOutputMin ? FMath::CeilToDouble(InOutputMin.GetValue() / InMinorGridStep) * InMinorGridStep 
+		                                      : FMath::FloorToDouble(ViewSpace.GetOutputMin() / InMinorGridStep) * InMinorGridStep;
+	const double LastMinorLine = InOutputMax ? FMath::FloorToDouble(InOutputMax.GetValue() / InMinorGridStep) * InMinorGridStep 
+		                                     : FMath::CeilToDouble(ViewSpace.GetOutputMax() / InMinorGridStep) * InMinorGridStep;
+
+	FNumberFormattingOptions FormattingOptions;
+	FormattingOptions.SetMaximumFractionalDigits(6);
+
+	// calculate min. distance between labels
+	const FSlateFontInfo FontInfo = FCoreStyle::Get().GetFontStyle("ToolTip.LargerFont");
+	const TSharedRef<FSlateFontMeasure> FontMeasureService = FSlateApplication::Get().GetRenderer()->GetFontMeasureService();
+	uint16 FontHeight = FontMeasureService->GetMaxCharacterHeight(FontInfo);
+
+	const double LabelDist = (1 / ViewSpace.PixelsPerOutput()) * (FontHeight + 3.0); // 3.0 for margin
+	double LineSkip = FMath::CeilToDouble(LabelDist / MajorGridStep) * MajorGridStep;
+	LineSkip = FMath::IsNearlyZero(LineSkip) ? KINDA_SMALL_NUMBER : LineSkip; // prevent mod by zero errors
+
+	for (double CurrentMinorLine = FirstMinorLine; CurrentMinorLine <= LastMinorLine; CurrentMinorLine += InMinorGridStep)
+	{
+		// check if is major grid line
+		if (FMath::IsNearlyZero(FMath::Fmod(FMath::Abs(CurrentMinorLine), MajorGridStep)))
+		{
+			OutMajorGridLines.Add(ViewSpace.ValueToScreen(CurrentMinorLine));
+			if (OutMajorGridLabels)
+			{
+				OutMajorGridLabels->Add(FMath::IsNearlyZero(FMath::Fmod(FMath::Abs(CurrentMinorLine), LineSkip))
+					? FText::Format(GridLineLabelFormatY, FText::AsNumber(CurrentMinorLine, &FormattingOptions)) 
+					: FText());
+			}
+		}
+		else
+		{
+			OutMinorGridLines.Add(ViewSpace.ValueToScreen(CurrentMinorLine));
+		}
+	}
+}
+
 } // namespace CurveEditor
\ No newline at end of file
diff --git a/Engine/Source/Editor/CurveEditor/Private/DragOperations/CurveEditorDragOperation_MoveKeys.cpp b/Engine/Source/Editor/CurveEditor/Private/DragOperations/CurveEditorDragOperation_MoveKeys.cpp
--- a/Engine/Source/Editor/CurveEditor/Private/DragOperations/CurveEditorDragOperation_MoveKeys.cpp
+++ b/Engine/Source/Editor/CurveEditor/Private/DragOperations/CurveEditorDragOperation_MoveKeys.cpp
@@ -67,6 +67,8 @@ void FCurveEditorDragOperation_MoveKeys::OnDrag(FVector2D InitialPosition, FVect
 		NewKeyPositionScratch.Reset();
 		NewKeyPositionScratch.Reserve(KeyData.StartKeyPositions.Num());
 
+		FCurveSnapMetrics SnapMetrics = CurveEditor->GetCurveSnapMetrics(KeyData.CurveID);
+
 		for (FKeyPosition StartPosition : KeyData.StartKeyPositions)
 		{
 			StartPosition.InputValue  += DeltaInput;
diff --git a/Engine/Source/Editor/CurveEditor/Private/DragOperations/ICurveEditorDragOperation.cpp b/Engine/Source/Editor/CurveEditor/Private/DragOperations/ICurveEditorDragOperation.cpp
--- a/Engine/Source/Editor/CurveEditor/Private/DragOperations/ICurveEditorDragOperation.cpp
+++ b/Engine/Source/Editor/CurveEditor/Private/DragOperations/ICurveEditorDragOperation.cpp
@@ -31,6 +31,6 @@ void ICurveEditorDragOperation::CancelDrag()
 
 void ICurveEditorKeyDragOperation::Initialize(FCurveEditor* InCurveEditor, const TOptional<FCurvePointHandle>& CardinalPoint)
 {
-	SnapMetrics = InCurveEditor->GetSnapMetrics();
+	// TODO: maybe cache snap data for all selected curves?
 	OnInitialize(InCurveEditor, CardinalPoint);
 }
\ No newline at end of file
diff --git a/Engine/Source/Editor/CurveEditor/Private/Filters/CurveEditorBakeFilter.cpp b/Engine/Source/Editor/CurveEditor/Private/Filters/CurveEditorBakeFilter.cpp
--- a/Engine/Source/Editor/CurveEditor/Private/Filters/CurveEditorBakeFilter.cpp
+++ b/Engine/Source/Editor/CurveEditor/Private/Filters/CurveEditorBakeFilter.cpp
@@ -20,9 +20,6 @@ void UCurveEditorBakeFilter::ApplyFilter_Impl(TSharedRef<FCurveEditor> InCurveEd
 	TArray<FKeyPosition> NewKeyPositions;
 	TArray<FKeyAttributes> NewKeyAttributes;
 
-	// Look at our Input Snap rate to determine how far apart the baked keys are.
-	FFrameRate BakeRate = InCurveEditor->GetSnapMetrics().InputSnapRate;
-
 	for (const TTuple<FCurveModelID, FKeyHandleSet>& Pair : InKeysToOperateOn)
 	{
 		FCurveModel* Curve = InCurveEditor->FindCurve(Pair.Key);
@@ -32,6 +29,9 @@ void UCurveEditorBakeFilter::ApplyFilter_Impl(TSharedRef<FCurveEditor> InCurveEd
 			continue;
 		}
 
+		// Look at our Input Snap rate to determine how far apart the baked keys are.
+		FFrameRate BakeRate = InCurveEditor->GetCurveSnapMetrics(Pair.Key).InputSnapRate;
+
 		KeyHandles.Reset(Pair.Value.Num());
 		KeyHandles.Append(Pair.Value.AsArray().GetData(), Pair.Value.Num());
 
diff --git a/Engine/Source/Editor/CurveEditor/Private/Filters/CurveEditorReduceFilter.cpp b/Engine/Source/Editor/CurveEditor/Private/Filters/CurveEditorReduceFilter.cpp
--- a/Engine/Source/Editor/CurveEditor/Private/Filters/CurveEditorReduceFilter.cpp
+++ b/Engine/Source/Editor/CurveEditor/Private/Filters/CurveEditorReduceFilter.cpp
@@ -64,7 +64,7 @@ void UCurveEditorReduceFilter::ApplyFilter_Impl(TSharedRef<FCurveEditor> InCurve
 	TArray<FKeyPosition> SelectedKeyPositions;
 	TArray<FKeyAttributes> SelectedKeyAttributes;
 
-	FFrameRate BakeRate = InCurveEditor->GetSnapMetrics().InputSnapRate;
+	FFrameRate BakeRate;
 
 	// Since we only remove keys we'll just copy the whole set and then remove the handles as well.
 	OutKeysToSelect = InKeysToOperateOn;
@@ -77,6 +77,8 @@ void UCurveEditorReduceFilter::ApplyFilter_Impl(TSharedRef<FCurveEditor> InCurve
 			continue;
 		}
 
+		BakeRate = InCurveEditor->GetCurveSnapMetrics(Pair.Key).InputSnapRate;
+
 		KeyHandles.Reset(Pair.Value.Num());
 		KeyHandles.Append(Pair.Value.AsArray().GetData(), Pair.Value.Num());
 
diff --git a/Engine/Source/Editor/CurveEditor/Private/SCurveEditorPanel.cpp b/Engine/Source/Editor/CurveEditor/Private/SCurveEditorPanel.cpp
--- a/Engine/Source/Editor/CurveEditor/Private/SCurveEditorPanel.cpp
+++ b/Engine/Source/Editor/CurveEditor/Private/SCurveEditorPanel.cpp
@@ -32,12 +32,13 @@
 #include "Filters/CurveEditorFilterBase.h"
 #include "Filters/CurveEditorBakeFilter.h"
 #include "Filters/CurveEditorReduceFilter.h"
-#include "SCurveValueSnapList.h"
+#include "SGridLineSpacingList.h"
 #include "Widgets/SFrameRatePicker.h"
 #include "CommonFrameRates.h"
 #include "CurveEditorViewRegistry.h"
 #include "SCurveEditorViewContainer.h"
 #include "Fonts/FontMeasure.h"
+#include "CurveEditorHelpers.h"
 
 #define LOCTEXT_NAMESPACE "SCurveEditorPanel"
 
@@ -1069,7 +1070,7 @@ TSharedPtr<FExtender> SCurveEditorPanel::GetToolbarExtender()
 				// Dropdown Menu to choose the snapping scale.
 				ToolBarBuilder.AddComboButton(
 					FUIAction(),
-					FOnGetContent::CreateSP(InEditorPanel, &SCurveEditorPanel::MakeValueSnapMenu),
+					FOnGetContent::CreateSP(InEditorPanel, &SCurveEditorPanel::MakeGridSpacingMenu),
 					LOCTEXT("ValueSnappingOptions", "Value Snapping"),
 					LOCTEXT("ValueSnappingOptionsToolTip", "Choose what precision the Value axis is snapped to while moving keys."),
 					TAttribute<FSlateIcon>(),
@@ -1186,27 +1187,28 @@ FText SCurveEditorPanel::GetTimeSnapMenuTooltip() const
 	return LOCTEXT("TimeSnappingOptionsToolTip", "Choose what precision the Time axis is snapped to while moving keys.");
 }
 
-TSharedRef<SWidget> SCurveEditorPanel::MakeValueSnapMenu()
+TSharedRef<SWidget> SCurveEditorPanel::MakeGridSpacingMenu()
 {
-	TArray<SCurveValueSnapList<float>::FNamedValue> ValueSnapAmounts;
+	TArray<SGridLineSpacingList::FNamedValue> SpacingAmounts;
 	// SnapValues.Add( SNumericDropDown<float>::FNamedValue( 0.001f, LOCTEXT( "Snap_OneThousandth", "0.001" ), LOCTEXT( "SnapDescription_OneThousandth", "Set snap to 1/1000th" ) ) );
 	//SnapValues.Add( SNumericDropDown<float>::FNamedValue( 0.01f, LOCTEXT( "Snap_OneHundredth", "0.01" ), LOCTEXT( "SnapDescription_OneHundredth", "Set snap to 1/100th" ) ) );
-	ValueSnapAmounts.Add(SCurveValueSnapList<float>::FNamedValue(0.1f, LOCTEXT("Snap_OneTenth", "0.1"), LOCTEXT("SnapDescription_OneTenth", "Set snap to 1/10th")));
-	ValueSnapAmounts.Add(SCurveValueSnapList<float>::FNamedValue(0.5f, LOCTEXT("Snap_OneHalf", "0.5"), LOCTEXT("SnapDescription_OneHalf", "Set snap to 1/2")));
-	ValueSnapAmounts.Add(SCurveValueSnapList<float>::FNamedValue(1.0f, LOCTEXT("Snap_One", "1"), LOCTEXT("SnapDescription_One", "Set snap to 1")));
-	ValueSnapAmounts.Add(SCurveValueSnapList<float>::FNamedValue(2.0f, LOCTEXT("Snap_Two", "2"), LOCTEXT("SnapDescription_Two", "Set snap to 2")));
-	ValueSnapAmounts.Add(SCurveValueSnapList<float>::FNamedValue(5.0f, LOCTEXT("Snap_Five", "5"), LOCTEXT("SnapDescription_Five", "Set snap to 5")));
-	ValueSnapAmounts.Add(SCurveValueSnapList<float>::FNamedValue(10.0f, LOCTEXT("Snap_Ten", "10"), LOCTEXT("SnapDescription_Ten", "Set snap to 10")));
-	ValueSnapAmounts.Add(SCurveValueSnapList<float>::FNamedValue(50.0f, LOCTEXT("Snap_Fifty", "50"), LOCTEXT("SnapDescription_50", "Set snap to 50")));
-	ValueSnapAmounts.Add(SCurveValueSnapList<float>::FNamedValue(100.0f, LOCTEXT("Snap_OneHundred", "100"), LOCTEXT("SnapDescription_OneHundred", "Set snap to 100")));
+	SpacingAmounts.Add(SGridLineSpacingList::FNamedValue(0.1f, LOCTEXT("OneTenth", "0.1"), LOCTEXT("Description_OneTenth", "Set grid spacing to 1/10th")));
+	SpacingAmounts.Add(SGridLineSpacingList::FNamedValue(0.5f, LOCTEXT("OneHalf", "0.5"), LOCTEXT("Description_OneHalf", "Set grid spacing to 1/2")));
+	SpacingAmounts.Add(SGridLineSpacingList::FNamedValue(1.0f, LOCTEXT("One", "1"), LOCTEXT("Description_One", "Set grid spacing to 1")));
+	SpacingAmounts.Add(SGridLineSpacingList::FNamedValue(2.0f, LOCTEXT("Two", "2"), LOCTEXT("Description_Two", "Set grid spacing to 2")));
+	SpacingAmounts.Add(SGridLineSpacingList::FNamedValue(5.0f, LOCTEXT("Five", "5"), LOCTEXT("Description_Five", "Set grid spacing to 5")));
+	SpacingAmounts.Add(SGridLineSpacingList::FNamedValue(10.0f, LOCTEXT("Ten", "10"), LOCTEXT("Description_Ten", "Set grid spacing to 10")));
+	SpacingAmounts.Add(SGridLineSpacingList::FNamedValue(50.0f, LOCTEXT("Fifty", "50"), LOCTEXT("Description_50", "Set grid spacing to 50")));
+	SpacingAmounts.Add(SGridLineSpacingList::FNamedValue(100.0f, LOCTEXT("OneHundred", "100"), LOCTEXT("Description_OneHundred", "Set grid spacing to 100")));
+	SpacingAmounts.Add(SGridLineSpacingList::FNamedValue(TOptional<float>(), LOCTEXT("Automatic", "Automatic"), LOCTEXT("Description_Automatic", "Set grid spacing to automatic")));
 
 	TSharedRef<SWidget> OutputSnapWidget =
-		SNew(SCurveValueSnapList<float>)
-		.DropDownValues(ValueSnapAmounts)
+		SNew(SGridLineSpacingList)
+		.DropDownValues(SpacingAmounts)
 		.MinDesiredValueWidth(60)
-		.Value_Lambda([this]() -> float { return this->CurveEditor->OutputSnapIntervalAttribute.Get(); })
-		.OnValueChanged_Lambda([this](float InNewOutputSnap) { this->CurveEditor->OutputSnapIntervalAttribute.Set(InNewOutputSnap); })
-		.HeaderText(LOCTEXT("CurveEditorMenuValueSnapHeader", "Value Snap"));
+		.Value_Lambda([this]() -> TOptional<float> { return this->CurveEditor->FixedGridSpacingAttribute.Get(); })
+		.OnValueChanged_Lambda([this](TOptional<float> InNewOutputSnap) { this->CurveEditor->FixedGridSpacingAttribute = InNewOutputSnap; })
+		.HeaderText(LOCTEXT("CurveEditorMenuGridSpacingHeader", "Grid Spacing"));
 
 	return OutputSnapWidget;
 }
diff --git a/Engine/Source/Editor/CurveEditor/Private/Views/SCurveEditorViewNormalized.cpp b/Engine/Source/Editor/CurveEditor/Private/Views/SCurveEditorViewNormalized.cpp
--- a/Engine/Source/Editor/CurveEditor/Private/Views/SCurveEditorViewNormalized.cpp
+++ b/Engine/Source/Editor/CurveEditor/Private/Views/SCurveEditorViewNormalized.cpp
@@ -28,7 +28,7 @@ void SCurveEditorViewNormalized::Construct(const FArguments& InArgs, TWeakPtr<FC
 	];
 }
 
-void SCurveEditorViewNormalized::GetGridLinesY(TSharedRef<FCurveEditor> CurveEditor, TArray<float>& MajorGridLines, TArray<float>& MinorGridLines, TArray<FText>& MajorGridLabels) const
+void SCurveEditorViewNormalized::GetGridLinesY(TSharedRef<const FCurveEditor> CurveEditor, TArray<float>& MajorGridLines, TArray<float>& MinorGridLines, TArray<FText>* MajorGridLabels) const
 {
 	FCurveEditorScreenSpace ViewSpace = GetViewSpace();
 
diff --git a/Engine/Source/Editor/CurveEditor/Private/Views/SCurveEditorViewNormalized.h b/Engine/Source/Editor/CurveEditor/Private/Views/SCurveEditorViewNormalized.h
--- a/Engine/Source/Editor/CurveEditor/Private/Views/SCurveEditorViewNormalized.h
+++ b/Engine/Source/Editor/CurveEditor/Private/Views/SCurveEditorViewNormalized.h
@@ -13,12 +13,12 @@ public:
 
 	void Construct(const FArguments& InArgs, TWeakPtr<FCurveEditor> InCurveEditor);
 
-	/** Tools should ignore vertical snapping because it causes issues with curves that have tiny extents. */
-	virtual bool IsValueSnapEnabled() const override { return false; }
+	/** Tools should use vertical snapping since grid lines to snap to will usually be visible */
+	virtual bool IsValueSnapEnabled() const override { return true; }
 
 private:
 
 	virtual void Tick(const FGeometry& AllottedGeometry, const double InCurrentTime, const float InDeltaTime) override;
 
-	virtual void GetGridLinesY(TSharedRef<FCurveEditor> CurveEditor, TArray<float>& MajorGridLines, TArray<float>& MinorGridLines, TArray<FText>& MajorGridLabels) const override;
+	virtual void GetGridLinesY(TSharedRef<const FCurveEditor> CurveEditor, TArray<float>& MajorGridLines, TArray<float>& MinorGridLines, TArray<FText>* MajorGridLabels) const override;
 };
diff --git a/Engine/Source/Editor/CurveEditor/Private/Views/SCurveEditorViewStacked.cpp b/Engine/Source/Editor/CurveEditor/Private/Views/SCurveEditorViewStacked.cpp
--- a/Engine/Source/Editor/CurveEditor/Private/Views/SCurveEditorViewStacked.cpp
+++ b/Engine/Source/Editor/CurveEditor/Private/Views/SCurveEditorViewStacked.cpp
@@ -6,7 +6,7 @@
 #include "SCurveEditorPanel.h"
 
 #include "Widgets/Text/STextBlock.h"
-
+#include "Algo/Copy.h"
 #include "EditorStyleSet.h"
 
 constexpr float StackedHeight = 150.f;
@@ -26,26 +26,50 @@ FVector2D SCurveEditorViewStacked::ComputeDesiredSize(float LayoutScaleMultiplie
 	return FVector2D(100.f, StackedHeight*CurveInfoByID.Num() + StackedPadding*(CurveInfoByID.Num()+1));
 }
 
-void SCurveEditorViewStacked::GetGridLinesY(TSharedRef<FCurveEditor> CurveEditor, TArray<float>& MajorGridLines, TArray<float>& MinorGridLines, TArray<FText>& MajorGridLabels) const
+void SCurveEditorViewStacked::GetGridLinesY(TSharedRef<const FCurveEditor> CurveEditor, TArray<float>& MajorGridLines, TArray<float>& MinorGridLines, TArray<FText>* MajorGridLabels) const
 {
+	FCurveEditorScreenSpace ViewSpace = GetViewSpace();
+
 	const double ValuePerPixel = 1.0 / StackedHeight;
 	const double ValueSpacePadding = StackedPadding * ValuePerPixel;
 
-	FCurveEditorScreenSpace ViewSpace = GetViewSpace();
-	for (int32 Index = 0; Index < CurveInfoByID.Num(); ++Index)
+	const TOptional<float> GridLineSpacing = CurveEditor->GetGridSpacing();
+	// draw standard lines if fixed spacing is not set
+	if (!GridLineSpacing)
 	{
-		double Padding = (Index + 1)*ValueSpacePadding;
-		double LowerValue = Index + Padding;
+		for (int32 Index = 0; Index < CurveInfoByID.Num(); ++Index)
+		{
+			double Padding = (Index + 1)*ValueSpacePadding;
+			double LowerValue = Index + Padding;
+
+			// Lower Grid line
+			MajorGridLines.Add(ViewSpace.ValueToScreen(LowerValue));
+			// Center Grid line
+			MajorGridLines.Add(ViewSpace.ValueToScreen(LowerValue + 0.5));
+			// Upper Grid line
+			MajorGridLines.Add(ViewSpace.ValueToScreen(LowerValue + 1.0));
+
+			MinorGridLines.Add(ViewSpace.ValueToScreen(LowerValue + 0.25));
+			MinorGridLines.Add(ViewSpace.ValueToScreen(LowerValue + 0.75));
+		}
+	}
+	else
+	{
+		TArray<FCurveModelID> CurveIDs;
+		CurveInfoByID.GetKeys(CurveIDs);
+		for (int32 Index = 0; Index < CurveInfoByID.Num(); ++Index)
+		{
+			double Padding = (Index + 1)*ValueSpacePadding;
+			double LowerValue = Index + Padding;
 
-		// Lower Grid line
-		MajorGridLines.Add(ViewSpace.ValueToScreen(LowerValue));
-		// Center Grid line
-		MajorGridLines.Add(ViewSpace.ValueToScreen(LowerValue + 0.5));
-		// Upper Grid line
-		MajorGridLines.Add(ViewSpace.ValueToScreen(LowerValue + 1.0));
+			double ViewSpaceMax = GetViewSpace().ValueToScreen(LowerValue);
+			double ViewSpaceMin = GetViewSpace().ValueToScreen(LowerValue + 1.0);
 
-		MinorGridLines.Add(ViewSpace.ValueToScreen(LowerValue + 0.25));
-		MinorGridLines.Add(ViewSpace.ValueToScreen(LowerValue + 0.75));
+			FCurveEditorScreenSpace CurveSpace = GetCurveSpace(CurveIDs[CurveInfoByID.Num() - Index - 1]);
+			TArray<float> TempMajorGridLines, TempMinorGridLines;
+			CurveEditor::ConstructFixedYGridLines(CurveSpace, 4, GridLineSpacing.GetValue(), MajorGridLines, MinorGridLines, CurveEditor->GetGridLineLabelFormatYAttribute().Get(), 
+				MajorGridLabels, CurveSpace.ScreenToValue(ViewSpaceMax), CurveSpace.ScreenToValue(ViewSpaceMin));
+		}
 	}
 }
 
@@ -84,7 +108,7 @@ void SCurveEditorViewStacked::DrawViewGrids(const FGeometry& AllottedGeometry, c
 	TArray<float> MajorGridLines, MinorGridLines;
 	TArray<FText> MajorGridLabels;
 
-	GetGridLinesX(CurveEditor.ToSharedRef(), MajorGridLines, MinorGridLines, MajorGridLabels);
+	GetGridLinesX(CurveEditor.ToSharedRef(), MajorGridLines, MinorGridLines, &MajorGridLabels);
 
 	// Pre-allocate an array of line points to draw our vertical lines. Each major grid line
 	// will overwrite the X value of both points but leave the Y value untouched so they draw from the bottom to the top.
@@ -129,27 +153,25 @@ void SCurveEditorViewStacked::DrawViewGrids(const FGeometry& AllottedGeometry, c
 
 		// Horizontal grid lines
 		{
+			TArray<float> MajorGridLinesH, MinorGridLinesH;
+			GetGridLinesY(CurveEditor.ToSharedRef(), MajorGridLinesH, MinorGridLinesH);
+
 			LinePoints[0].X = 0.0;
 			LinePoints[1].X = Width;
 
-			// Top grid line
-			LinePoints[0].Y = LinePoints[1].Y = PixelTop;
-			FSlateDrawElement::MakeLines(OutDrawElements, GridLineLayerId, PaintGeometry, LinePoints, DrawEffects, MajorGridColor, false);
-
-			// Center line
-			LinePoints[0].Y = LinePoints[1].Y = ViewSpace.ValueToScreen(LowerValue + 0.5);
-			FSlateDrawElement::MakeLines(OutDrawElements, GridLineLayerId, PaintGeometry, LinePoints, DrawEffects, MajorGridColor, false);
-
-			// Bottom line
-			LinePoints[0].Y = LinePoints[1].Y = PixelBottom;
-			FSlateDrawElement::MakeLines(OutDrawElements, GridLineLayerId, PaintGeometry, LinePoints, DrawEffects, MajorGridColor, false);
-
-			// Minor intermediate lines
-			LinePoints[0].Y = LinePoints[1].Y = ViewSpace.ValueToScreen(LowerValue + 0.25);
-			FSlateDrawElement::MakeLines(OutDrawElements, GridLineLayerId, PaintGeometry, LinePoints, DrawEffects, MinorGridColor, false);
+			// draw major lines
+			for (float GridLineVal : MajorGridLinesH)
+			{
+				LinePoints[0].Y = LinePoints[1].Y = GridLineVal;
+				FSlateDrawElement::MakeLines(OutDrawElements, GridLineLayerId, PaintGeometry, LinePoints, DrawEffects, MajorGridColor, false);
+			}
 
-			LinePoints[0].Y = LinePoints[1].Y = ViewSpace.ValueToScreen(LowerValue + 0.75);
-			FSlateDrawElement::MakeLines(OutDrawElements, GridLineLayerId, PaintGeometry, LinePoints, DrawEffects, MinorGridColor, false);
+			// draw minor lines
+			for (float GridLineVal : MinorGridLinesH)
+			{
+				LinePoints[0].Y = LinePoints[1].Y = GridLineVal;
+				FSlateDrawElement::MakeLines(OutDrawElements, GridLineLayerId, PaintGeometry, LinePoints, DrawEffects, MinorGridColor, false);
+			}
 		}
 
 		// Vertical grid lines
diff --git a/Engine/Source/Editor/CurveEditor/Private/Views/SCurveEditorViewStacked.h b/Engine/Source/Editor/CurveEditor/Private/Views/SCurveEditorViewStacked.h
--- a/Engine/Source/Editor/CurveEditor/Private/Views/SCurveEditorViewStacked.h
+++ b/Engine/Source/Editor/CurveEditor/Private/Views/SCurveEditorViewStacked.h
@@ -13,8 +13,8 @@ public:
 
 	void Construct(const FArguments& InArgs, TWeakPtr<FCurveEditor> InCurveEditor);
 	
-	/** Tools should ignore vertical snapping because it causes issues with curves that have tiny extents. */
-	virtual bool IsValueSnapEnabled() const override { return false; }
+	/** Tools should use vertical snapping since grid lines to snap to will usually be visible */
+	virtual bool IsValueSnapEnabled() const override { return true; }
 
 private:
 
@@ -22,7 +22,7 @@ private:
 	virtual FVector2D ComputeDesiredSize(float LayoutScaleMultiplier) const override;
 
 	virtual void PaintView(const FPaintArgs& Args, const FGeometry& AllottedGeometry, const FSlateRect& MyCullingRect, FSlateWindowElementList& OutDrawElements, int32 BaseLayerId, const FWidgetStyle& InWidgetStyle, bool bParentEnabled) const override;
-	virtual void GetGridLinesY(TSharedRef<FCurveEditor> CurveEditor, TArray<float>& MajorGridLines, TArray<float>& MinorGridLines, TArray<FText>& MajorGridLabels) const override;
+	virtual void GetGridLinesY(TSharedRef<const FCurveEditor> CurveEditor, TArray<float>& MajorGridLines, TArray<float>& MinorGridLines, TArray<FText>* MajorGridLabels = nullptr) const override;
 
 	void DrawViewGrids(const FGeometry& AllottedGeometry, const FSlateRect& MyCullingRect, FSlateWindowElementList& OutDrawElements, int32 BaseLayerId, ESlateDrawEffect DrawEffects) const;
 	void DrawLabels(const FGeometry& AllottedGeometry, const FSlateRect& MyCullingRect, FSlateWindowElementList& OutDrawElements, int32 BaseLayerId, ESlateDrawEffect DrawEffects) const;
diff --git a/Engine/Source/Editor/CurveEditor/Private/Views/SInteractiveCurveEditorView.cpp b/Engine/Source/Editor/CurveEditor/Private/Views/SInteractiveCurveEditorView.cpp
--- a/Engine/Source/Editor/CurveEditor/Private/Views/SInteractiveCurveEditorView.cpp
+++ b/Engine/Source/Editor/CurveEditor/Private/Views/SInteractiveCurveEditorView.cpp
@@ -134,7 +134,7 @@ FSlateColor SInteractiveCurveEditorView::GetCurveCaptionColor() const
 	return BackgroundTint.CopyWithNewOpacity(1.f);
 }
 
-void SInteractiveCurveEditorView::GetGridLinesX(TSharedRef<FCurveEditor> CurveEditor, TArray<float>& MajorGridLines, TArray<float>& MinorGridLines, TArray<FText>& MajorGridLabels) const
+void SInteractiveCurveEditorView::GetGridLinesX(TSharedRef<const FCurveEditor> CurveEditor, TArray<float>& MajorGridLines, TArray<float>& MinorGridLines, TArray<FText>* MajorGridLabels) const
 {
 	CurveEditor->GetGridLinesX(MajorGridLines, MinorGridLines, MajorGridLabels);
 
@@ -156,9 +156,17 @@ void SInteractiveCurveEditorView::GetGridLinesX(TSharedRef<FCurveEditor> CurveEd
 	}
 }
 
-void SInteractiveCurveEditorView::GetGridLinesY(TSharedRef<FCurveEditor> CurveEditor, TArray<float>& MajorGridLines, TArray<float>& MinorGridLines, TArray<FText>& MajorGridLabels) const
+void SInteractiveCurveEditorView::GetGridLinesY(TSharedRef<const FCurveEditor> CurveEditor, TArray<float>& MajorGridLines, TArray<float>& MinorGridLines, TArray<FText>* MajorGridLabels) const
 {
-	CurveEditor::ConstructYGridLines(GetViewSpace(), 4, MajorGridLines, MinorGridLines, CurveEditor->GetGridLineLabelFormatYAttribute().Get(), &MajorGridLabels);
+	const TOptional<float> GridLineSpacing = CurveEditor->GetGridSpacing();
+	if (!GridLineSpacing)
+	{
+		CurveEditor::ConstructYGridLines(GetViewSpace(), 4, MajorGridLines, MinorGridLines, CurveEditor->GetGridLineLabelFormatYAttribute().Get(), MajorGridLabels);
+	}
+	else
+	{
+		CurveEditor::ConstructFixedYGridLines(GetViewSpace(), 4, GridLineSpacing.GetValue(), MajorGridLines, MinorGridLines, CurveEditor->GetGridLineLabelFormatYAttribute().Get(), MajorGridLabels, TOptional<double>(), TOptional<double>());
+	}
 }
 
 int32 SInteractiveCurveEditorView::OnPaint(const FPaintArgs& Args, const FGeometry& AllottedGeometry, const FSlateRect& MyCullingRect, FSlateWindowElementList& OutDrawElements, int32 BaseLayerId, const FWidgetStyle& InWidgetStyle, bool bParentEnabled) const
@@ -211,7 +219,7 @@ void SInteractiveCurveEditorView::DrawGridLines(TSharedRef<FCurveEditor> CurveEd
 	TArray<float> MajorGridLines, MinorGridLines;
 	TArray<FText> MajorGridLabels;
 
-	GetGridLinesX(CurveEditor, MajorGridLines, MinorGridLines, MajorGridLabels);
+	GetGridLinesX(CurveEditor, MajorGridLines, MinorGridLines, &MajorGridLabels);
 	ensureMsgf(MajorGridLabels.Num() == 0 || MajorGridLines.Num() == MajorGridLabels.Num(), TEXT("If grid labels are specified, one must be specified for every major grid line, even if it is just an empty FText."));
 
 	// Pre-allocate an array of line points to draw our vertical lines. Each major grid line
@@ -294,7 +302,7 @@ void SInteractiveCurveEditorView::DrawGridLines(TSharedRef<FCurveEditor> CurveEd
 	MajorGridLines.Reset();
 	MinorGridLines.Reset();
 	MajorGridLabels.Reset();
-	GetGridLinesY(CurveEditor, MajorGridLines, MinorGridLines, MajorGridLabels);
+	GetGridLinesY(CurveEditor, MajorGridLines, MinorGridLines, &MajorGridLabels);
 	ensureMsgf(MajorGridLabels.Num() == 0 || MajorGridLines.Num() == MajorGridLabels.Num(), TEXT("If grid labels are specified, one must be specified for every major grid line, even if it is just an empty FText."));
 
 	// Reset our cached Line to draw from left to right
@@ -691,7 +699,7 @@ void SInteractiveCurveEditorView::UpdateCurveProximities(FVector2D MousePixel)
 		{
 			FCurveEditorScreenSpace CurveSpace = GetCurveSpace(CurveProximities[0].Get<0>());
 			double MouseTime = CurveSpace.ScreenToSeconds(MousePixel.X) - HoveredCurve->GetInputDisplayOffset();
-			double EvaluatedTime = CurveEditor->GetSnapMetrics().SnapInputSeconds(MouseTime);
+			double EvaluatedTime = CurveEditor->GetCurveSnapMetrics(CurveProximities[0].Get<0>()).SnapInputSeconds(MouseTime);
 
 			double EvaluatedValue = 0.0;
 			HoveredCurve->Evaluate(EvaluatedTime, EvaluatedValue);
@@ -884,8 +892,9 @@ FReply SInteractiveCurveEditorView::OnMouseButtonDown(const FGeometry& MyGeometr
 					double MouseTime = CurveSpace.ScreenToSeconds(MousePixel.X);
 					double MouseValue = CurveSpace.ScreenToValue(MousePixel.Y);
 
-					MouseTime = CurveEditor->GetSnapMetrics().SnapInputSeconds(MouseTime);
-					MouseValue = CurveEditor->GetSnapMetrics().SnapOutput(MouseValue);
+					FCurveSnapMetrics SnapMetrics = CurveEditor->GetCurveSnapMetrics(HoveredCurve.GetValue());
+					MouseTime = SnapMetrics.SnapInputSeconds(MouseTime);
+					MouseValue = SnapMetrics.SnapOutput(MouseValue);
 
 					// When adding to a curve with no variance, add it with the same value so that
 					// curves don't pop wildly in normalized views due to a slight difference between the keys
@@ -1301,9 +1310,6 @@ void SInteractiveCurveEditorView::AddKeyAtTime(const TSet<FCurveModelID>& ToCurv
 		return;
 	}
 
-	// Ensure the time is snapped if needed
-	InTime = CurveEditor->GetSnapMetrics().SnapInputSeconds(InTime);
-
 	FScopedTransaction Transaction(LOCTEXT("AddKeyAtTime", "Add Key"));
 	bool bAddedKey = false;
 
@@ -1317,14 +1323,18 @@ void SInteractiveCurveEditorView::AddKeyAtTime(const TSet<FCurveModelID>& ToCurv
 		FCurveModel* CurveModel = CurveEditor->FindCurve(CurveModelID);
 		check(CurveModel);
 
+		// Ensure the time is snapped if needed
+		FCurveSnapMetrics SnapMetrics = CurveEditor->GetCurveSnapMetrics(CurveModelID);
+		double SnappedTime = SnapMetrics.SnapInputSeconds(InTime);
+
 		// Support optional input display offsets 
-		double EvalTime = InTime - CurveModel->GetInputDisplayOffset();
+		double EvalTime = SnappedTime - CurveModel->GetInputDisplayOffset();
 
 		double CurveValue = 0.0;
 		if (CurveModel->Evaluate(EvalTime, CurveValue))
 		{
 			CurveModel->Modify();
-			CurveValue = CurveEditor->GetSnapMetrics().SnapOutput(CurveValue);
+			CurveValue = SnapMetrics.SnapOutput(CurveValue);
 
 			// Curve Models allow us to create new keys ontop of existing keys which works, but causes some user confusion
 			// Before we create a key, we instead check to see if there is already a key at this time, and if there is, we
diff --git a/Engine/Source/Editor/CurveEditor/Public/CurveEditor.h b/Engine/Source/Editor/CurveEditor/Public/CurveEditor.h
--- a/Engine/Source/Editor/CurveEditor/Public/CurveEditor.h
+++ b/Engine/Source/Editor/CurveEditor/Public/CurveEditor.h
@@ -22,6 +22,7 @@
 #include "EditorUndoClient.h"
 #include "CurveEditorScreenSpace.h"
 #include "ITimeSlider.h"
+#include "CurveEditorHelpers.h"
 
 class UCurveEditorSettings;
 class FUICommandList;
@@ -66,8 +67,8 @@ public:
 	/** Attribute used to retrieve the current input snap rate (also used for display) */
 	TAttribute<FFrameRate> InputSnapRateAttribute;
 
-	/** Attribute used to retrieve the current output snap interval */
-	TAttribute<double> OutputSnapIntervalAttribute;
+	/** Attribute used to retrieve the current value-axis grid line state */
+	TAttribute<TOptional<float>> FixedGridSpacingAttribute;
 
 	/** Attribute used to determine if we should snap input values */
 	TAttribute<bool> InputSnapEnabledAttribute;
@@ -228,7 +229,13 @@ public:
 	/**
 	 * Generate a utility struct for snapping values
 	 */
-	FCurveEditorSnapMetrics GetSnapMetrics() const;
+	FCurveSnapMetrics GetCurveSnapMetrics(FCurveModelID CurveModel) const;
+
+	/**
+	 * Returns the value grid line spacing state
+	 */
+	TOptional<float> GetGridSpacing() const { return FixedGridSpacingAttribute.Get(); }
+
 	/**
 	 * Returned the cached struct for snapping editing movement to a specific axis based on user preferences.
 	 */
@@ -383,7 +390,7 @@ public:
 	 * function for the Y grid lines is not provided due to the Curve Editor's ability to have multiple
 	 * views with repeated gridlines and values.
 	 */
-	virtual void GetGridLinesX(TArray<float>& MajorGridLines, TArray<float>& MinorGridLines, TArray<FText>& MajorGridLabels) const
+	virtual void GetGridLinesX(TArray<float>& MajorGridLines, TArray<float>& MinorGridLines, TArray<FText>* MajorGridLabels) const
 	{
 		ConstructXGridLines(MajorGridLines, MinorGridLines, MajorGridLabels);
 	}
@@ -398,7 +405,7 @@ protected:
 	/**
 	 * Construct grid lines along the current display frame rate or time-base
 	 */
-	void ConstructXGridLines(TArray<float>& MajorGridLines, TArray<float>& MinorGridLines, TArray<FText>& MajorGridLabels) const;
+	void ConstructXGridLines(TArray<float>& MajorGridLines, TArray<float>& MinorGridLines, TArray<FText>* MajorGridLabels) const;
 
 	/**
 	 * Internal zoom to fit implementation
diff --git a/Engine/Source/Editor/CurveEditor/Public/CurveEditorHelpers.h b/Engine/Source/Editor/CurveEditor/Public/CurveEditorHelpers.h
--- a/Engine/Source/Editor/CurveEditor/Public/CurveEditorHelpers.h
+++ b/Engine/Source/Editor/CurveEditor/Public/CurveEditorHelpers.h
@@ -15,4 +15,7 @@ namespace CurveEditor
 
 	void ConstructYGridLines(const FCurveEditorScreenSpace& ViewSpace, uint8 InMinorDivisions, TArray<float>& OutMajorGridLines, TArray<float>& OutMinorGridLines, FText GridLineLabelFormatY, TArray<FText>* OutMajorGridLabels);
 
+	void ConstructFixedYGridLines(const FCurveEditorScreenSpace& ViewSpace, uint8 InMinorDivisions, double InMinorGridStep, TArray<float>& OutMajorGridLines, TArray<float>& OutMinorGridLines, FText GridLineLabelFormatY, 
+		TArray<FText>* OutMajorGridLabels, TOptional<double> InOutputMin, TOptional<double> InOutputMax);
+
 } // namespace CurveEditor
\ No newline at end of file
diff --git a/Engine/Source/Editor/CurveEditor/Public/CurveEditorSnapMetrics.h b/Engine/Source/Editor/CurveEditor/Public/CurveEditorSnapMetrics.h
--- a/Engine/Source/Editor/CurveEditor/Public/CurveEditorSnapMetrics.h
+++ b/Engine/Source/Editor/CurveEditor/Public/CurveEditorSnapMetrics.h
@@ -4,18 +4,14 @@
 
 #include "Misc/FrameRate.h"
 #include "Math/Axis.h"
+#include "Runtime/Core/Public/Algo/MinElement.h"
 
-
-/**
- * Utility struct that acts as a cache of the current snapping metrics for the curve editor
- */
-struct FCurveEditorSnapMetrics
+struct FCurveSnapMetrics
 {
-	FCurveEditorSnapMetrics()
+	FCurveSnapMetrics()
 	{
 		bSnapOutputValues = 0;
 		bSnapInputValues = 0;
-		OutputSnapInterval = 1.0;
 	}
 
 	/** Whether we are snapping to the output snap interval */
@@ -24,8 +20,8 @@ struct FCurveEditorSnapMetrics
 	/** Whether we are snapping to the input snap rate */
 	uint8 bSnapInputValues : 1;
 
-	/** The output snap interval */
-	double OutputSnapInterval;
+	/** Grid lines to snap to */
+	TArray<double> AllGridLines;
 
 	/** The input snap rate */
 	FFrameRate InputSnapRate;
@@ -35,11 +31,13 @@ struct FCurveEditorSnapMetrics
 	{
 		return bSnapInputValues && InputSnapRate.IsValid() ? (InputTime * InputSnapRate).RoundToFrame() / InputSnapRate : InputTime;
 	}
-
+	
 	/** Snap the specified output value to the output snap interval if necessary */
 	FORCEINLINE double SnapOutput(double OutputValue)
 	{
-		return bSnapOutputValues && OutputSnapInterval != 0.f ? FMath::RoundToDouble(OutputValue / OutputSnapInterval) * OutputSnapInterval : OutputValue;
+		return bSnapOutputValues ? *Algo::MinElement(AllGridLines, 
+			[OutputValue](double Val1, double Val2) { return FMath::Abs(Val1 - OutputValue) < FMath::Abs(Val2 - OutputValue); }
+		) : OutputValue;
 	}
 };
 
diff --git a/Engine/Source/Editor/CurveEditor/Public/ICurveEditorDragOperation.h b/Engine/Source/Editor/CurveEditor/Public/ICurveEditorDragOperation.h
--- a/Engine/Source/Editor/CurveEditor/Public/ICurveEditorDragOperation.h
+++ b/Engine/Source/Editor/CurveEditor/Public/ICurveEditorDragOperation.h
@@ -78,9 +78,6 @@ class ICurveEditorKeyDragOperation : public ICurveEditorDragOperation
 {
 public:
 
-	/** Cached (and potentially manipulated) snap metrics to be used for this drag */
-	FCurveEditorSnapMetrics SnapMetrics;
-
 	/**
 	 * Initialize this drag operation using the specified curve editor pointer and an optional cardinal point
 
diff --git a/Engine/Source/Editor/CurveEditor/Public/SCurveEditorPanel.h b/Engine/Source/Editor/CurveEditor/Public/SCurveEditorPanel.h
--- a/Engine/Source/Editor/CurveEditor/Public/SCurveEditorPanel.h
+++ b/Engine/Source/Editor/CurveEditor/Public/SCurveEditorPanel.h
@@ -127,7 +127,8 @@ private:
 
 	TSharedRef<SWidget> MakeTimeSnapMenu();
 	FText GetTimeSnapMenuTooltip() const;
-	TSharedRef<SWidget> MakeValueSnapMenu();
+
+	TSharedRef<SWidget> MakeGridSpacingMenu();
 	TSharedRef<SWidget> MakeAxisSnapMenu();
 	TSharedRef<SWidget> MakeViewModeMenu();
 
diff --git a/Engine/Source/Editor/CurveEditor/Public/SCurveEditorView.h b/Engine/Source/Editor/CurveEditor/Public/SCurveEditorView.h
--- a/Engine/Source/Editor/CurveEditor/Public/SCurveEditorView.h
+++ b/Engine/Source/Editor/CurveEditor/Public/SCurveEditorView.h
@@ -173,6 +173,9 @@ public:
 		return true;
 	}
 
+	virtual void GetGridLinesX(TSharedRef<const FCurveEditor> CurveEditor, TArray<float>& MajorGridLines, TArray<float>& MinorGridLines, TArray<FText>* MajorGridLabels = nullptr) const {}
+	virtual void GetGridLinesY(TSharedRef<const FCurveEditor> CurveEditor, TArray<float>& MajorGridLines, TArray<float>& MinorGridLines, TArray<FText>* MajorGridLabels = nullptr) const {}
+
 protected:
 
 	// ~SWidget interface
diff --git a/Engine/Source/Editor/CurveEditor/Public/SCurveValueSnapList.h b/Engine/Source/Editor/CurveEditor/Public/SCurveValueSnapList.h
deleted file mode 100644
diff --git a/Engine/Source/Editor/CurveEditor/Public/SGridLineSpacingList.h b/Engine/Source/Editor/CurveEditor/Public/SGridLineSpacingList.h
new file mode 100644
--- /dev/null
+++ b/Engine/Source/Editor/CurveEditor/Public/SGridLineSpacingList.h
@@ -0,0 +1,188 @@
+// Copyright 1998-2019 Epic Games, Inc. All Rights Reserved.
+
+#pragma once
+
+#include "CoreMinimal.h"
+#include "Misc/Attribute.h"
+#include "Layout/Visibility.h"
+#include "Input/Events.h"
+#include "Widgets/SWidget.h"
+#include "Textures/SlateIcon.h"
+#include "Layout/Margin.h"
+#include "Widgets/DeclarativeSyntaxSupport.h"
+#include "Widgets/SCompoundWidget.h"
+#include "Widgets/SBoxPanel.h"
+#include "Framework/Commands/UIAction.h"
+#include "Framework/MultiBox/MultiBoxBuilder.h"
+#include "Widgets/Text/STextBlock.h"
+#include "Framework/Application/SlateApplication.h"
+#include "Widgets/Input/SEditableTextBox.h"
+#include "Widgets/Input/SComboButton.h"
+
+/** A widget which creates a list of pre-determined numbers and a spot to enter a custom number. Similar to SNumericList but without the dropdown menu. */
+class SGridLineSpacingList : public SCompoundWidget
+{
+public:
+	DECLARE_DELEGATE_OneParam( FOnValueChanged, TOptional<float> );
+
+public:
+	/** Represents a named numeric value for display in the drop down menu. */ 
+	class FNamedValue
+	{
+	public:
+		/** Creates a new FNamedValue 
+		  * @InValue The numeric value to be assigned
+		  * @InName The display text for the value in the UI 
+		  * @InDescription The description of the value used in tooltips or wherever a longer description is needed. */
+		FNamedValue( TOptional<float> InValue, FText InName, FText InDescription )
+		{
+			Value = InValue;
+			Name = InName;
+			Description = InDescription;
+		}
+
+		TOptional<float> GetValue() const
+		{
+			return Value;
+		}
+
+		FText GetName() const
+		{
+			return Name;
+		}
+
+		FText GetDescription() const
+		{
+			return Description;
+		}
+
+	private:
+		TOptional<float> Value;
+		FText Name;
+		FText Description;
+	};
+
+public:
+	SLATE_BEGIN_ARGS( SGridLineSpacingList )
+		: _DropDownValues()
+		, _MinDesiredValueWidth( 40)
+		, _bShowNamedValue(true)
+		{}
+
+		/** The values which are used to populate the drop down menu. */
+		SLATE_ARGUMENT( TArray<FNamedValue>, DropDownValues )
+		/** Controls the minimum width for the text box portion of the control. */
+		SLATE_ATTRIBUTE( float, MinDesiredValueWidth )
+		/** The value displayed by the control. */
+		SLATE_ATTRIBUTE( TOptional<float>, Value )
+		/** If enabled, attempt to show the original name specified in the Drop Down values for a given entry instead of just a straight FText::AsNumber(NumericType). */
+		SLATE_ATTRIBUTE(bool, bShowNamedValue)
+		/** An optional header to prepend to the generated list. */
+		SLATE_ATTRIBUTE( FText, HeaderText)
+		/** The callback for when the value changes. */
+		SLATE_EVENT( FOnValueChanged, OnValueChanged )
+
+	SLATE_END_ARGS()
+
+	void Construct( const FArguments& InArgs )
+	{
+		DropDownValues = InArgs._DropDownValues;
+		Value = InArgs._Value;
+		OnValueChanged = InArgs._OnValueChanged;
+
+		FMenuBuilder MenuBuilder(true, nullptr);
+
+		// Allow an optional header text because it looks nicer and helps give context to the widget.
+		if (InArgs._HeaderText.IsSet())
+		{
+			MenuBuilder.BeginSection("Header", InArgs._HeaderText.Get());
+		}
+
+		for (const FNamedValue& NamedValue : DropDownValues)
+		{
+			AddMenuEntry(MenuBuilder, NamedValue);
+		}
+
+		// Add the custom entry box at the end. 
+		MenuBuilder.AddWidget(
+			SNew(SBox)
+			.HAlign(HAlign_Right)
+			.Padding(FMargin(8, 0, 0, 0))
+			[
+				SNew(SEditableTextBox)
+				.RevertTextOnEscape(true)
+				.SelectAllTextWhenFocused(true)
+				.Text(this, &SGridLineSpacingList::GetValueText)
+				.OnTextCommitted(this, &SGridLineSpacingList::ValueTextCommitted)
+				.MinDesiredWidth(InArgs._MinDesiredValueWidth)
+			],
+			NSLOCTEXT("SNumericList", "CustomNumberDisplayLabel", "Custom")
+		);
+
+		ChildSlot
+		[
+			MenuBuilder.MakeWidget()
+		];
+	}
+
+private:
+	void AddMenuEntry(FMenuBuilder& MenuBuilder, const FNamedValue& Info)
+	{
+		MenuBuilder.AddMenuEntry(
+			Info.GetName(),
+			Info.GetDescription(),
+			FSlateIcon(),
+			FUIAction(
+				FExecuteAction::CreateSP(this, &SGridLineSpacingList::SetValue, Info.GetValue()),
+				FCanExecuteAction(),
+				FIsActionChecked::CreateSP(this, &SGridLineSpacingList::IsSameValue, Info.GetValue())
+			),
+			NAME_None,
+			EUserInterfaceActionType::RadioButton
+		);
+	}
+
+	bool IsSameValue(TOptional<float> InValue) const
+	{
+		return Value.Get() == InValue;
+	}
+
+	/** Get the value text for displaying in the custom entry box. We use a textbox here so we can show actual names from the original list if possible. */
+	FText GetValueText() const
+	{
+		if (bShowNamedValue.Get())
+		{
+			for ( FNamedValue DropDownValue : DropDownValues )
+			{
+				if (Value.Get() == DropDownValue.GetValue())
+				{
+					return DropDownValue.GetName();
+				}
+			}
+		}
+
+		return FText::FromString(LexToSanitizedString(Value.Get() ? Value.Get().GetValue() : 0.0f));
+	}
+
+	void ValueTextCommitted(const FText& InNewText, ETextCommit::Type InTextCommit)
+	{
+		if (InNewText.IsNumeric())
+		{
+			TOptional<float> NewValue = 0.0f;
+			TTypeFromString<float>::FromString(NewValue.GetValue(), *InNewText.ToString());
+			OnValueChanged.ExecuteIfBound(NewValue);
+		}
+	}
+
+	void SetValue( TOptional<float> InValue )
+	{
+		FSlateApplication::Get().ClearKeyboardFocus( EFocusCause::Cleared );
+		OnValueChanged.ExecuteIfBound(InValue);
+	}
+
+private:
+	TArray<FNamedValue> DropDownValues;
+	TAttribute<bool> bShowNamedValue;
+	TAttribute<TOptional<float>> Value;
+	FOnValueChanged OnValueChanged;
+};
diff --git a/Engine/Source/Editor/CurveEditor/Public/Views/SInteractiveCurveEditorView.h b/Engine/Source/Editor/CurveEditor/Public/Views/SInteractiveCurveEditorView.h
--- a/Engine/Source/Editor/CurveEditor/Public/Views/SInteractiveCurveEditorView.h
+++ b/Engine/Source/Editor/CurveEditor/Public/Views/SInteractiveCurveEditorView.h
@@ -71,14 +71,14 @@ public:
 
 	void Construct(const FArguments& InArgs, TWeakPtr<FCurveEditor> InCurveEditor);
 
+	virtual void GetGridLinesX(TSharedRef<const FCurveEditor> CurveEditor, TArray<float>& MajorGridLines, TArray<float>& MinorGridLines, TArray<FText>* MajorGridLabels = nullptr) const override;
+	virtual void GetGridLinesY(TSharedRef<const FCurveEditor> CurveEditor, TArray<float>& MajorGridLines, TArray<float>& MinorGridLines, TArray<FText>* MajorGridLabels = nullptr) const override;
+
 protected:
 
 	// ~SCurveEditorView Interface
 	virtual void GetPointsWithinWidgetRange(const FSlateRect& WidgetRectangle, TArray<FCurvePointHandle>* OutPoints) const override;
 
-	virtual void GetGridLinesX(TSharedRef<FCurveEditor> CurveEditor, TArray<float>& MajorGridLines, TArray<float>& MinorGridLines, TArray<FText>& MajorGridLabels) const;
-	virtual void GetGridLinesY(TSharedRef<FCurveEditor> CurveEditor, TArray<float>& MajorGridLines, TArray<float>& MinorGridLines, TArray<FText>& MajorGridLabels) const;
-
 	virtual void PaintView(const FPaintArgs& Args, const FGeometry& AllottedGeometry, const FSlateRect& MyCullingRect, FSlateWindowElementList& OutDrawElements, int32 BaseLayerId, const FWidgetStyle& InWidgetStyle, bool bParentEnabled) const;
 
 protected:
diff --git a/Engine/Source/Editor/MovieSceneTools/Private/Channels/SCurveEditorEventChannelView.cpp b/Engine/Source/Editor/MovieSceneTools/Private/Channels/SCurveEditorEventChannelView.cpp
--- a/Engine/Source/Editor/MovieSceneTools/Private/Channels/SCurveEditorEventChannelView.cpp
+++ b/Engine/Source/Editor/MovieSceneTools/Private/Channels/SCurveEditorEventChannelView.cpp
@@ -30,7 +30,7 @@ FVector2D SCurveEditorEventChannelView::ComputeDesiredSize(float LayoutScaleMult
 	return FVector2D(100.f, TrackHeight * (CurveInfoByID.Num()));
 }
 
-void SCurveEditorEventChannelView::GetGridLinesY(TSharedRef<FCurveEditor> CurveEditor, TArray<float>& MajorGridLines, TArray<float>& MinorGridLines, TArray<FText>& MajorGridLabels) const
+void SCurveEditorEventChannelView::GetGridLinesY(TSharedRef<const FCurveEditor> CurveEditor, TArray<float>& MajorGridLines, TArray<float>& MinorGridLines, TArray<FText>* MajorGridLabels) const
 {}
 
 void SCurveEditorEventChannelView::Tick(const FGeometry& AllottedGeometry, const double InCurrentTime, const float InDeltaTime)
diff --git a/Engine/Source/Editor/MovieSceneTools/Private/Channels/SCurveEditorEventChannelView.h b/Engine/Source/Editor/MovieSceneTools/Private/Channels/SCurveEditorEventChannelView.h
--- a/Engine/Source/Editor/MovieSceneTools/Private/Channels/SCurveEditorEventChannelView.h
+++ b/Engine/Source/Editor/MovieSceneTools/Private/Channels/SCurveEditorEventChannelView.h
@@ -10,6 +10,8 @@ public:
 
 	void Construct(const FArguments& InArgs, TWeakPtr<FCurveEditor> InCurveEditor);
 
+	virtual void GetGridLinesY(TSharedRef<const FCurveEditor> CurveEditor, TArray<float>& MajorGridLines, TArray<float>& MinorGridLines, TArray<FText>* MajorGridLabels) const override;
+
 private:
 
 	//~ SWidget Interface
@@ -17,7 +19,6 @@ private:
 	virtual FVector2D ComputeDesiredSize(float LayoutScaleMultiplier) const override;
 
 	//~ SInteractiveCurveEditorView interface
-	virtual void GetGridLinesY(TSharedRef<FCurveEditor> CurveEditor, TArray<float>& MajorGridLines, TArray<float>& MinorGridLines, TArray<FText>& MajorGridLabels) const override;
 	virtual void PaintView(const FPaintArgs& Args, const FGeometry& AllottedGeometry, const FSlateRect& MyCullingRect, FSlateWindowElementList& OutDrawElements, int32 BaseLayerId, const FWidgetStyle& InWidgetStyle, bool bParentEnabled) const override;
 
 	void DrawLabels(const FGeometry& AllottedGeometry, FSlateWindowElementList& OutDrawElements, int32 BaseLayerId, ESlateDrawEffect DrawEffects) const;
diff --git a/Engine/Source/Editor/Sequencer/Private/Sequencer.cpp b/Engine/Source/Editor/Sequencer/Private/Sequencer.cpp
--- a/Engine/Source/Editor/Sequencer/Private/Sequencer.cpp
+++ b/Engine/Source/Editor/Sequencer/Private/Sequencer.cpp
@@ -249,7 +249,7 @@ public:
 		: WeakSequencer(InSequencer)
 	{}
 
-	virtual void GetGridLinesX(TArray<float>& MajorGridLines, TArray<float>& MinorGridLines, TArray<FText>& MajorGridLabels) const override
+	virtual void GetGridLinesX(TArray<float>& MajorGridLines, TArray<float>& MinorGridLines, TArray<FText>* MajorGridLabels) const override
 	{
 		TSharedPtr<FSequencer> Sequencer = WeakSequencer.Pin();
 		FCurveEditorScreenSpaceH PanelInputSpace = GetPanelInputSpace();
@@ -328,7 +328,7 @@ void FSequencer::InitSequencer(const FSequencerInitParams& InitParams, const TSh
 		CurveEditorModel->OutputSnapEnabledAttribute  = MakeAttributeLambda([this]{ return Settings->GetSnapCurveValueToInterval(); });
 		CurveEditorModel->OnOutputSnapEnabledChanged  = FOnSetBoolean::CreateLambda([this](bool NewValue){ Settings->SetSnapCurveValueToInterval(NewValue); });
 
-		CurveEditorModel->OutputSnapIntervalAttribute = MakeAttributeLambda([this]{ return (double)Settings->GetCurveValueSnapInterval(); });
+		CurveEditorModel->FixedGridSpacingAttribute   = MakeAttributeLambda([this]() -> TOptional<float> { return Settings->GetGridSpacing(); });
 		CurveEditorModel->InputSnapRateAttribute      = MakeAttributeSP(this, &FSequencer::GetFocusedDisplayRate);
 
 		CurveEditorModel->DefaultKeyAttributes        = MakeAttributeSP(this, &FSequencer::GetDefaultKeyAttributes);
diff --git a/Engine/Source/Editor/Sequencer/Private/SequencerSettings.cpp b/Engine/Source/Editor/Sequencer/Private/SequencerSettings.cpp
--- a/Engine/Source/Editor/Sequencer/Private/SequencerSettings.cpp
+++ b/Engine/Source/Editor/Sequencer/Private/SequencerSettings.cpp
@@ -26,6 +26,7 @@ USequencerSettings::USequencerSettings( const FObjectInitializer& ObjectInitiali
 	bSnapPlayTimeToPressedKey = true;
 	bSnapPlayTimeToDraggedKey = true;
 	CurveValueSnapInterval = 10.0f;
+	GridSpacing = TOptional<float>();
 	bSnapCurveValueToInterval = true;
 	bLabelBrowserVisible = false;
 	bShowSelectedNodesOnly = false;
@@ -312,6 +313,20 @@ void USequencerSettings::SetCurveValueSnapInterval( float InCurveValueSnapInterv
 	}
 }
 
+TOptional<float> USequencerSettings::GetGridSpacing() const
+{
+	return GridSpacing;
+}
+
+void USequencerSettings::SetGridSpacing(TOptional<float> InGridSpacing)
+{
+	if (InGridSpacing != GridSpacing)
+	{
+		GridSpacing = InGridSpacing;
+		SaveConfig();
+	}
+}
+
 bool USequencerSettings::GetSnapCurveValueToInterval() const
 {
 	return bSnapCurveValueToInterval;
diff --git a/Engine/Source/Editor/Sequencer/Public/SequencerSettings.h b/Engine/Source/Editor/Sequencer/Public/SequencerSettings.h
--- a/Engine/Source/Editor/Sequencer/Public/SequencerSettings.h
+++ b/Engine/Source/Editor/Sequencer/Public/SequencerSettings.h
@@ -188,6 +188,11 @@ public:
 	/** Sets the snapping interval for curve values. */
 	void SetCurveValueSnapInterval(float InCurveValueSnapInterval);
 
+	/** Gets the state for spacing between grid lines */
+	TOptional<float> GetGridSpacing() const;
+	/** Sets the grid line spacing state */
+	void SetGridSpacing(TOptional<float> InGridSpacing);
+
 	/** Gets whether or not to snap curve values to the interval. */
 	bool GetSnapCurveValueToInterval() const;
 	/** Sets whether or not to snap curve values to the interval. */
@@ -404,6 +409,9 @@ protected:
 	/** The curve value interval to snap to. */
 	float CurveValueSnapInterval;
 
+	/** grid line spacing state */
+	TOptional<float> GridSpacing;
+
 	/** Enable or disable snapping the curve value to the curve value interval. */
 	UPROPERTY( config, EditAnywhere, Category=Snapping )
 	bool bSnapCurveValueToInterval;
diff --git a/Engine/Source/Runtime/Core/Private/Misc/FrameRate.cpp b/Engine/Source/Runtime/Core/Private/Misc/FrameRate.cpp
--- a/Engine/Source/Runtime/Core/Private/Misc/FrameRate.cpp
+++ b/Engine/Source/Runtime/Core/Private/Misc/FrameRate.cpp
@@ -175,7 +175,7 @@ FText FFrameRate::ToPrettyText() const
 	}
 }
 
-bool FFrameRate::ComputeGridSpacing(float PixelsPerSecond, double& OutMajorInterval, int32& OutMinorDivisions, float MinTickPx, float DesiredMajorTickPx)
+bool FFrameRate::ComputeGridSpacing(float PixelsPerSecond, double& OutMajorInterval, int32& OutMinorDivisions, float MinTickPx, float DesiredMajorTickPx) const
 {
 	if (PixelsPerSecond <= 0.f)
 	{
diff --git a/Engine/Source/Runtime/Core/Public/Misc/FrameRate.h b/Engine/Source/Runtime/Core/Public/Misc/FrameRate.h
--- a/Engine/Source/Runtime/Core/Public/Misc/FrameRate.h
+++ b/Engine/Source/Runtime/Core/Public/Misc/FrameRate.h
@@ -135,7 +135,7 @@ struct FFrameRate
 	 * @param DesiredMajorTickPx (Optional) The desired size to compute major tick lines from
 	 * @return True if a valid grid spacing was computed, false otherwise.
 	 */
-	CORE_API bool ComputeGridSpacing(const float PixelsPerSecond, double& OutMajorInterval, int32& OutMinorDivisions, float MinTickPx = 30.f, float DesiredMajorTickPx = 120.f);
+	CORE_API bool ComputeGridSpacing(const float PixelsPerSecond, double& OutMajorInterval, int32& OutMinorDivisions, float MinTickPx = 30.f, float DesiredMajorTickPx = 120.f) const;
 
 	/**
 	 * Get the maximum number of seconds representable with this framerate
