Curve Editor: Show buffered curves in all views

diff --git a/Engine/Source/Editor/CurveEditor/Private/CurveEditor.cpp b/Engine/Source/Editor/CurveEditor/Private/CurveEditor.cpp
--- a/Engine/Source/Editor/CurveEditor/Private/CurveEditor.cpp
+++ b/Engine/Source/Editor/CurveEditor/Private/CurveEditor.cpp
@@ -22,6 +22,7 @@
 #include "Framework/Notifications/NotificationManager.h"
 #include "Widgets/Notifications/SNotificationList.h"
 #include "Runtime/Core/Public/Algo/Transform.h"
+#include "SCurveEditor.h" // for access to LogCurveEditor
 
 #define LOCTEXT_NAMESPACE "CurveEditor"
 
@@ -933,31 +934,29 @@ void FCurveEditor::SetBufferedCurves(const TSet<FCurveModelID>& InCurves)
 		FCurveModel* CurveModel = FindCurve(CurveID);
 		check(CurveModel);
 
-
-		// We loop through the key data on the curve and apply it to the copy, since there's no copy constructors for Curve Models (for good reason!)
-		FBufferedCurve& BufferedCurve = BufferedCurves.Emplace_GetRef();
-		
-		int32 NumKeys = CurveModel->GetNumKeys();
-		BufferedCurve.KeyPositions.SetNumUninitialized(NumKeys);
-		BufferedCurve.KeyAttributes.SetNumUninitialized(NumKeys);
-
-		// Copy all of the from the curve into our cached buffer
-		TArray<FKeyHandle> KeyHandles;
-
-		CurveModel->GetKeys(*this, TNumericLimits<double>::Lowest(), TNumericLimits<double>::Max(), TNumericLimits<double>::Lowest(), TNumericLimits<double>::Max(), KeyHandles);
-		CurveModel->GetKeyPositions(KeyHandles, BufferedCurve.KeyPositions);
-		CurveModel->GetKeyAttributes(KeyHandles, BufferedCurve.KeyAttributes);
-		BufferedCurve.IntentionName = CurveModel->GetIntentionName();
+		// Add a buffered curve copy if the curve model supports buffered curves
+		TUniquePtr<IBufferedCurveModel> CurveModelCopy = CurveModel->CreateBufferedCurveCopy();
+		if (CurveModelCopy) 
+		{ 
+			BufferedCurves.Add(MoveTemp(CurveModelCopy)); 
+		}
+		else
+		{
+			UE_LOG(LogCurveEditor, Warning, TEXT("Failed to buffer curve, curve model did not provide a copy."))
+		}
 	}
 }
 
 
-void FCurveEditor::ApplyBufferedCurveToTarget(const FBufferedCurve& BufferedCurve, FCurveModel* TargetCurve)
+void FCurveEditor::ApplyBufferedCurveToTarget(const IBufferedCurveModel* BufferedCurve, FCurveModel* TargetCurve)
 {
 	check(TargetCurve);
+	check(BufferedCurve);
 
 	TArray<FKeyPosition> KeyPositions;
 	TArray<FKeyAttributes> KeyAttributes;
+	BufferedCurve->GetKeyPositions(KeyPositions);
+	BufferedCurve->GetKeyAttributes(KeyAttributes);
 
 
 	// Copy the data from the Buffered curve into the target curve. This just does wholesale replacement.
@@ -968,7 +967,7 @@ void FCurveEditor::ApplyBufferedCurveToTarget(const FBufferedCurve& BufferedCurv
 	TargetCurve->RemoveKeys(TargetKeyHandles);
 
 	// Now put our buffered keys into the target curve
-	TargetCurve->AddKeys(BufferedCurve.KeyPositions, BufferedCurve.KeyAttributes);
+	TargetCurve->AddKeys(KeyPositions, KeyAttributes);
 }
 
 bool FCurveEditor::ApplyBufferedCurves(const TSet<FCurveModelID>& InCurvesToApplyTo)
@@ -1020,7 +1019,7 @@ bool FCurveEditor::ApplyBufferedCurves(const TSet<FCurveModelID>& InCurvesToAppl
 		int32 MatchedBufferedCurveIndex = -1;
 		for (int32 BufferedCurveIndex = BufferedCurveSearchIndexStart; BufferedCurveIndex < BufferedCurves.Num(); BufferedCurveIndex++)
 		{
-			if (BufferedCurves[BufferedCurveIndex].IntentionName == TargetIntent)
+			if (BufferedCurves[BufferedCurveIndex]->GetIntentionName() == TargetIntent)
 			{
 				MatchedBufferedCurveIndex = BufferedCurveIndex;
 
@@ -1045,7 +1044,7 @@ bool FCurveEditor::ApplyBufferedCurves(const TSet<FCurveModelID>& InCurvesToAppl
 			NumCurvesMatchedByIntent++;
 			bFoundAnyMatchedIntent = true;
 
-			const FBufferedCurve& BufferedCurve = BufferedCurves[MatchedBufferedCurveIndex];
+			const IBufferedCurveModel* BufferedCurve = BufferedCurves[MatchedBufferedCurveIndex].Get();
 			ApplyBufferedCurveToTarget(BufferedCurve, TargetCurve);
 
 		}
@@ -1097,7 +1096,7 @@ bool FCurveEditor::ApplyBufferedCurves(const TSet<FCurveModelID>& InCurvesToAppl
 		
 		for (int32 CurveIndex = 0; CurveIndex < InCurvesToApplyTo.Num(); CurveIndex++)
 		{
-			ApplyBufferedCurveToTarget(BufferedCurves[CurveIndex], FindCurve(CurvesToApplyTo[CurveIndex]));
+			ApplyBufferedCurveToTarget(BufferedCurves[CurveIndex].Get(), FindCurve(CurvesToApplyTo[CurveIndex]));
 		}
 
 		FText NotificationText;
diff --git a/Engine/Source/Editor/CurveEditor/Private/RichCurveEditorModel.cpp b/Engine/Source/Editor/CurveEditor/Private/RichCurveEditorModel.cpp
--- a/Engine/Source/Editor/CurveEditor/Private/RichCurveEditorModel.cpp
+++ b/Engine/Source/Editor/CurveEditor/Private/RichCurveEditorModel.cpp
@@ -13,6 +13,93 @@
 #include "EditorStyleSet.h"
 #include "UObject/Package.h"
 
+
+void RefineCurvePoints(const FRichCurve* RichCurve, double TimeThreshold, float ValueThreshold, TArray<TTuple<double, double>>& InOutPoints)
+{
+	check(RichCurve);
+	const float InterpTimes[] = { 0.25f, 0.5f, 0.6f };
+
+	for (int32 Index = 0; Index < InOutPoints.Num() - 1; ++Index)
+	{
+		TTuple<double, double> Lower = InOutPoints[Index];
+		TTuple<double, double> Upper = InOutPoints[Index + 1];
+
+		if ((Upper.Get<0>() - Lower.Get<0>()) >= TimeThreshold)
+		{
+			bool bSegmentIsLinear = true;
+
+			TTuple<double, double> Evaluated[ARRAY_COUNT(InterpTimes)];
+
+			for (int32 InterpIndex = 0; InterpIndex < ARRAY_COUNT(InterpTimes); ++InterpIndex)
+			{
+				double& EvalTime = Evaluated[InterpIndex].Get<0>();
+
+				EvalTime = FMath::Lerp(Lower.Get<0>(), Upper.Get<0>(), InterpTimes[InterpIndex]);
+
+				float Value = RichCurve->Eval(EvalTime);
+
+				const float LinearValue = FMath::Lerp(Lower.Get<1>(), Upper.Get<1>(), InterpTimes[InterpIndex]);
+				if (bSegmentIsLinear)
+				{
+					bSegmentIsLinear = FMath::IsNearlyEqual(Value, LinearValue, ValueThreshold);
+				}
+
+				Evaluated[InterpIndex].Get<1>() = Value;
+			}
+
+			if (!bSegmentIsLinear)
+			{
+				// Add the point
+				InOutPoints.Insert(Evaluated, ARRAY_COUNT(Evaluated), Index + 1);
+				--Index;
+			}
+		}
+	}
+}
+
+/**
+ * Buffered curve implementation for a rich curve, stores a copy of the rich curve in order to draw itself.
+ */
+class FRichBufferedCurveModel : public IBufferedCurveModel
+{
+public:
+	FRichBufferedCurveModel(FRichCurve* InRichCurve, TArray<FKeyPosition>&& InKeyPositions, TArray<FKeyAttributes>&& InKeyAttributes,
+		const FString& InIntentionName, const double InValueMin, const double InValueMax)
+		: IBufferedCurveModel(MoveTemp(InKeyPositions), MoveTemp(InKeyAttributes), InIntentionName, InValueMin, InValueMax)
+		, RichCurve(*InRichCurve)
+	{}
+
+	virtual void DrawCurve(const FCurveEditor& CurveEditor, const FCurveEditorScreenSpace& ScreenSpace, TArray<TTuple<double, double>>& InterpolatingPoints) const override
+	{
+		const double StartTimeSeconds = ScreenSpace.GetInputMin();
+		const double EndTimeSeconds = ScreenSpace.GetInputMax();
+		const double TimeThreshold = FMath::Max(0.0001, 1.0 / ScreenSpace.PixelsPerInput());
+		const double ValueThreshold = FMath::Max(0.0001, 1.0 / ScreenSpace.PixelsPerOutput());
+
+		InterpolatingPoints.Add(MakeTuple(StartTimeSeconds, double(RichCurve.Eval(StartTimeSeconds))));
+
+		for (const FRichCurveKey& Key : RichCurve.GetConstRefOfKeys())
+		{
+			if (Key.Time > StartTimeSeconds && Key.Time < EndTimeSeconds)
+			{
+				InterpolatingPoints.Add(MakeTuple(double(Key.Time), double(Key.Value)));
+			}
+		}
+
+		InterpolatingPoints.Add(MakeTuple(EndTimeSeconds, double(RichCurve.Eval(EndTimeSeconds))));
+
+		int32 OldSize = InterpolatingPoints.Num();
+		do
+		{
+			OldSize = InterpolatingPoints.Num();
+			RefineCurvePoints(&RichCurve, TimeThreshold, ValueThreshold, InterpolatingPoints);
+		} while (OldSize != InterpolatingPoints.Num());
+	}
+
+private:
+	FRichCurve RichCurve;
+};
+
 FRichCurveEditorModel::FRichCurveEditorModel(FRichCurve* InRichCurve, UObject* InOwner)
 	: RichCurve(InRichCurve), WeakOwner(InOwner)
 {
@@ -89,48 +176,6 @@ void FRichCurveEditorModel::RemoveKeys(TArrayView<const FKeyHandle> InKeys)
 	}
 }
 
-void RefineCurvePoints(const FRichCurve* RichCurve, double TimeThreshold, float ValueThreshold, TArray<TTuple<double, double>>& InOutPoints)
-{
-	const float InterpTimes[] = { 0.25f, 0.5f, 0.6f };
-
-	for (int32 Index = 0; Index < InOutPoints.Num() - 1; ++Index)
-	{
-		TTuple<double, double> Lower = InOutPoints[Index];
-		TTuple<double, double> Upper = InOutPoints[Index + 1];
-
-		if ((Upper.Get<0>() - Lower.Get<0>()) >= TimeThreshold)
-		{
-			bool bSegmentIsLinear = true;
-
-			TTuple<double, double> Evaluated[ARRAY_COUNT(InterpTimes)];
-
-			for (int32 InterpIndex = 0; InterpIndex < ARRAY_COUNT(InterpTimes); ++InterpIndex)
-			{
-				double& EvalTime  = Evaluated[InterpIndex].Get<0>();
-
-				EvalTime = FMath::Lerp(Lower.Get<0>(), Upper.Get<0>(), InterpTimes[InterpIndex]);
-
-				float Value = RichCurve->Eval(EvalTime);
-
-				const float LinearValue = FMath::Lerp(Lower.Get<1>(), Upper.Get<1>(), InterpTimes[InterpIndex]);
-				if (bSegmentIsLinear)
-				{
-					bSegmentIsLinear = FMath::IsNearlyEqual(Value, LinearValue, ValueThreshold);
-				}
-
-				Evaluated[InterpIndex].Get<1>() = Value;
-			}
-
-			if (!bSegmentIsLinear)
-			{
-				// Add the point
-				InOutPoints.Insert(Evaluated, ARRAY_COUNT(Evaluated), Index+1);
-				--Index;
-			}
-		}
-	}
-}
-
 void FRichCurveEditorModel::DrawCurve(const FCurveEditor& CurveEditor, const FCurveEditorScreenSpace& ScreenSpace, TArray<TTuple<double, double>>& InOutPoints) const
 {
 	if (UObject* Owner = WeakOwner.Get())
@@ -471,6 +516,30 @@ void FRichCurveEditorModel::CreateKeyProxies(TArrayView<const FKeyHandle> InKeyH
 	}
 }
 
+TUniquePtr<IBufferedCurveModel> FRichCurveEditorModel::CreateBufferedCurveCopy() const
+{
+	if (UObject* Owner = WeakOwner.Get())
+	{
+		TArray<FKeyHandle> TargetKeyHandles;
+		for (auto It = RichCurve->GetKeyHandleIterator(); It; ++It)
+		{
+			TargetKeyHandles.Add(*It);
+		}
+
+		TArray<FKeyPosition> KeyPositions;
+		TArray<FKeyAttributes> KeyAttributes;
+		GetKeyPositions(TargetKeyHandles, KeyPositions);
+		GetKeyAttributes(TargetKeyHandles, KeyAttributes);
+
+		double ValueMin = 0.f, ValueMax = 1.f;
+		GetValueRange(ValueMin, ValueMax);
+
+		return MakeUnique<FRichBufferedCurveModel>(RichCurve, MoveTemp(KeyPositions), MoveTemp(KeyAttributes), GetIntentionName(), ValueMin, ValueMax);
+	}
+
+	return nullptr;
+}
+
 void FRichCurveEditorModel::GetTimeRange(double& MinTime, double& MaxTime) const
 {
 	if (UObject* Owner = WeakOwner.Get())
diff --git a/Engine/Source/Editor/CurveEditor/Private/Views/SCurveEditorViewNormalized.cpp b/Engine/Source/Editor/CurveEditor/Private/Views/SCurveEditorViewNormalized.cpp
--- a/Engine/Source/Editor/CurveEditor/Private/Views/SCurveEditorViewNormalized.cpp
+++ b/Engine/Source/Editor/CurveEditor/Private/Views/SCurveEditorViewNormalized.cpp
@@ -40,6 +40,85 @@ void SCurveEditorViewNormalized::GetGridLinesY(TSharedRef<const FCurveEditor> Cu
 	MinorGridLines.Add(ViewSpace.ValueToScreen(0.75));
 }
 
+FTransform2D CalculateViewToCurveTransform(const double InCurveOutputMin, const double InCurveOutputMax)
+{
+	const double Scale = (InCurveOutputMax - InCurveOutputMin);
+	if (InCurveOutputMax > InCurveOutputMin)
+	{
+		return FTransform2D(FScale2D(1.f, Scale), FVector2D(0.f, InCurveOutputMin));
+	}
+	else
+	{
+		return FVector2D(0.f, InCurveOutputMin - 0.5);
+	}
+}
+
+void SCurveEditorViewNormalized::DrawBufferedCurves(const FGeometry& AllottedGeometry, const FSlateRect& MyCullingRect, FSlateWindowElementList& OutDrawElements, int32 BaseLayerId, ESlateDrawEffect DrawEffects) const
+{
+	TSharedPtr<FCurveEditor> CurveEditor = WeakCurveEditor.Pin();
+	if (!CurveEditor)
+	{
+		return;
+	}
+
+	const float BufferedCurveThickness = 1.f;
+	const bool  bAntiAliasCurves = true;
+	const FLinearColor CurveColor = CurveViewConstants::BufferedCurveColor;
+	const int32 CurveLayerId = BaseLayerId + CurveViewConstants::ELayerOffset::Curves;
+
+	const double ValuePerPixel = 1.0 / AllottedGeometry.GetLocalSize().Y;
+	const double ValueSpacePadding = NormalizedPadding * ValuePerPixel;
+
+	// Calculate the normalized view to curve transform for each buffered curve, then draw
+	const TArray<TUniquePtr<IBufferedCurveModel>>& BufferedCurves = CurveEditor->GetBufferedCurves();
+	for (const TUniquePtr<IBufferedCurveModel>& BufferedCurve : BufferedCurves)
+	{
+		FTransform2D ViewToBufferedCurveTransform;
+		double CurveOutputMin = BufferedCurve->GetValueMin(), CurveOutputMax = BufferedCurve->GetValueMax();
+
+		ViewToBufferedCurveTransform = CalculateViewToCurveTransform(CurveOutputMin, CurveOutputMax);
+
+		TArray<TTuple<double, double>> CurveSpaceInterpolatingPoints;
+		FCurveEditorScreenSpace CurveSpace = GetViewSpace().ToCurveSpace(ViewToBufferedCurveTransform);
+
+		BufferedCurve->DrawCurve(*CurveEditor, CurveSpace, CurveSpaceInterpolatingPoints);
+
+		TArray<FVector2D> ScreenSpaceInterpolatingPoints;
+		for (TTuple<double, double> Point : CurveSpaceInterpolatingPoints)
+		{
+			ScreenSpaceInterpolatingPoints.Add(FVector2D(
+				CurveSpace.SecondsToScreen(Point.Get<0>()),
+				CurveSpace.ValueToScreen(Point.Get<1>())
+			));
+		}
+
+		FSlateDrawElement::MakeLines(
+			OutDrawElements,
+			CurveLayerId,
+			AllottedGeometry.ToPaintGeometry(),
+			ScreenSpaceInterpolatingPoints,
+			DrawEffects,
+			CurveColor,
+			bAntiAliasCurves,
+			BufferedCurveThickness
+		);
+	}
+}
+
+void SCurveEditorViewNormalized::PaintView(const FPaintArgs& Args, const FGeometry& AllottedGeometry, const FSlateRect& MyCullingRect, FSlateWindowElementList& OutDrawElements, int32 BaseLayerId, const FWidgetStyle& InWidgetStyle, bool bParentEnabled) const
+{
+	TSharedPtr<FCurveEditor> CurveEditor = WeakCurveEditor.Pin();
+	if (CurveEditor)
+	{
+		const ESlateDrawEffect DrawEffects = ShouldBeEnabled(bParentEnabled) ? ESlateDrawEffect::None : ESlateDrawEffect::DisabledEffect;
+
+		DrawBackground(AllottedGeometry, OutDrawElements, BaseLayerId, DrawEffects);
+		DrawGridLines(CurveEditor.ToSharedRef(), AllottedGeometry, OutDrawElements, BaseLayerId, DrawEffects);
+		DrawBufferedCurves(AllottedGeometry, MyCullingRect, OutDrawElements, BaseLayerId, DrawEffects);
+		DrawCurves(CurveEditor.ToSharedRef(), AllottedGeometry, MyCullingRect, OutDrawElements, BaseLayerId, InWidgetStyle, DrawEffects);
+	}
+}
+
 void SCurveEditorViewNormalized::Tick(const FGeometry& AllottedGeometry, const double InCurrentTime, const float InDeltaTime)
 {
 	TSharedPtr<FCurveEditor> CurveEditor = WeakCurveEditor.Pin();
@@ -67,15 +146,7 @@ void SCurveEditorViewNormalized::Tick(const FGeometry& AllottedGeometry, const d
 			double CurveOutputMin = 0, CurveOutputMax = 1;
 			Curve->GetValueRange(CurveOutputMin, CurveOutputMax);
 
-			const double Scale = (CurveOutputMax - CurveOutputMin);
-			if (CurveOutputMax > CurveOutputMin)
-			{
-				It->Value.ViewToCurveTransform = FTransform2D(FScale2D(1.f, Scale), FVector2D(0.f, CurveOutputMin));
-			}
-			else
-			{
-				It->Value.ViewToCurveTransform = FVector2D(0.f, CurveOutputMin-0.5);
-			}
+			It->Value.ViewToCurveTransform = CalculateViewToCurveTransform(CurveOutputMin, CurveOutputMax);
 		}
 	}
 
diff --git a/Engine/Source/Editor/CurveEditor/Private/Views/SCurveEditorViewNormalized.h b/Engine/Source/Editor/CurveEditor/Private/Views/SCurveEditorViewNormalized.h
--- a/Engine/Source/Editor/CurveEditor/Private/Views/SCurveEditorViewNormalized.h
+++ b/Engine/Source/Editor/CurveEditor/Private/Views/SCurveEditorViewNormalized.h
@@ -20,5 +20,9 @@ private:
 
 	virtual void Tick(const FGeometry& AllottedGeometry, const double InCurrentTime, const float InDeltaTime) override;
 
+	virtual void PaintView(const FPaintArgs& Args, const FGeometry& AllottedGeometry, const FSlateRect& MyCullingRect, FSlateWindowElementList& OutDrawElements, int32 BaseLayerId, const FWidgetStyle& InWidgetStyle, bool bParentEnabled) const override;
+
 	virtual void GetGridLinesY(TSharedRef<const FCurveEditor> CurveEditor, TArray<float>& MajorGridLines, TArray<float>& MinorGridLines, TArray<FText>* MajorGridLabels) const override;
+
+	void DrawBufferedCurves(const FGeometry& AllottedGeometry, const FSlateRect& MyCullingRect, FSlateWindowElementList& OutDrawElements, int32 BaseLayerId, ESlateDrawEffect DrawEffects) const;
 };
diff --git a/Engine/Source/Editor/CurveEditor/Private/Views/SCurveEditorViewStacked.cpp b/Engine/Source/Editor/CurveEditor/Private/Views/SCurveEditorViewStacked.cpp
--- a/Engine/Source/Editor/CurveEditor/Private/Views/SCurveEditorViewStacked.cpp
+++ b/Engine/Source/Editor/CurveEditor/Private/Views/SCurveEditorViewStacked.cpp
@@ -57,15 +57,17 @@ void SCurveEditorViewStacked::GetGridLinesY(TSharedRef<const FCurveEditor> Curve
 	{
 		TArray<FCurveModelID> CurveIDs;
 		CurveInfoByID.GetKeys(CurveIDs);
-		for (int32 Index = 0; Index < CurveInfoByID.Num(); ++Index)
+		for (auto It = CurveInfoByID.CreateConstIterator(); It; ++It)
 		{
+			const int32 Index = CurveInfoByID.Num() - It->Value.CurveIndex - 1;
+
 			double Padding = (Index + 1)*ValueSpacePadding;
 			double LowerValue = Index + Padding;
 
 			double ViewSpaceMax = GetViewSpace().ValueToScreen(LowerValue);
 			double ViewSpaceMin = GetViewSpace().ValueToScreen(LowerValue + 1.0);
 
-			FCurveEditorScreenSpace CurveSpace = GetCurveSpace(CurveIDs[CurveInfoByID.Num() - Index - 1]);
+			FCurveEditorScreenSpace CurveSpace = GetCurveSpace(CurveIDs[Index]);
 			TArray<float> TempMajorGridLines, TempMinorGridLines;
 			CurveEditor::ConstructFixedYGridLines(CurveSpace, 4, GridLineSpacing.GetValue(), MajorGridLines, MinorGridLines, CurveEditor->GetGridLineLabelFormatYAttribute().Get(), 
 				MajorGridLabels, CurveSpace.ScreenToValue(ViewSpaceMax), CurveSpace.ScreenToValue(ViewSpaceMin));
@@ -83,6 +85,7 @@ void SCurveEditorViewStacked::PaintView(const FPaintArgs& Args, const FGeometry&
 		DrawBackground(AllottedGeometry, OutDrawElements, BaseLayerId, DrawEffects);
 		DrawViewGrids(AllottedGeometry, MyCullingRect, OutDrawElements, BaseLayerId, DrawEffects);
 		DrawLabels(AllottedGeometry, MyCullingRect, OutDrawElements, BaseLayerId, DrawEffects);
+		DrawBufferedCurves(AllottedGeometry, MyCullingRect, OutDrawElements, BaseLayerId, DrawEffects);
 		DrawCurves(CurveEditor.ToSharedRef(), AllottedGeometry, MyCullingRect, OutDrawElements, BaseLayerId, InWidgetStyle, DrawEffects);
 	}
 }
@@ -134,7 +137,6 @@ void SCurveEditorViewStacked::DrawViewGrids(const FGeometry& AllottedGeometry, c
 
 		double PixelBottom = ViewSpace.ValueToScreen(LowerValue);
 		double PixelTop    = ViewSpace.ValueToScreen(LowerValue + 1.0);
-
 		if (!FSlateRect::DoRectanglesIntersect(MyCullingRect, TransformRect(AllottedGeometry.GetAccumulatedLayoutTransform(), FSlateRect(0, PixelTop, Width, PixelBottom))))
 		{
 			continue;
@@ -256,6 +258,88 @@ void SCurveEditorViewStacked::DrawLabels(const FGeometry& AllottedGeometry, cons
 	}
 }
 
+FTransform2D CalculateViewToCurveTransform(const double InCurveOutputMin, double const InCurveOutputMax, const double InValueOffset)
+{
+	if (InCurveOutputMax > InCurveOutputMin)
+	{
+		return Concatenate(FVector2D(0.f, InValueOffset), Concatenate(FScale2D(1.f, (InCurveOutputMax - InCurveOutputMin)), FVector2D(0.f, InCurveOutputMin)));
+	}
+	else
+	{
+		return Concatenate(FVector2D(0.f, InValueOffset - 0.5), FVector2D(0.f, InCurveOutputMin));
+	}
+}
+
+void SCurveEditorViewStacked::DrawBufferedCurves(const FGeometry& AllottedGeometry, const FSlateRect& MyCullingRect, FSlateWindowElementList& OutDrawElements, int32 BaseLayerId, ESlateDrawEffect DrawEffects) const
+{
+	TSharedPtr<FCurveEditor> CurveEditor = WeakCurveEditor.Pin();
+	if (!CurveEditor)
+	{
+		return;
+	}
+
+	const TArray<TUniquePtr<IBufferedCurveModel>>& BufferedCurves = CurveEditor->GetBufferedCurves();
+
+	const float BufferedCurveThickness = 1.f;
+	const bool  bAntiAliasCurves = true;
+	const FLinearColor CurveColor = CurveViewConstants::BufferedCurveColor;
+	const int32 CurveLayerId = BaseLayerId + CurveViewConstants::ELayerOffset::Curves;
+
+	const double ValuePerPixel = 1.0 / StackedHeight;
+	const double ValueSpacePadding = StackedPadding * ValuePerPixel;
+
+
+	// draw the buffered curves for each view
+	for (auto It = CurveInfoByID.CreateConstIterator(); It; ++It)
+	{
+		FCurveModel* Curve = CurveEditor->FindCurve(It.Key());
+		if (!ensureAlways(Curve))
+		{
+			continue;
+		}
+
+		FTransform2D ViewToBufferedCurveTransform;
+
+		const int32  CurveIndexFromBottom = CurveInfoByID.Num() - It->Value.CurveIndex - 1;
+		const double PaddingToBottomOfView = (CurveIndexFromBottom + 1)*ValueSpacePadding;
+		const double ValueOffset = -CurveIndexFromBottom - PaddingToBottomOfView;
+
+		// Calculate the view to curve transform for each buffered curve, then draw
+		for (const TUniquePtr<IBufferedCurveModel>& BufferedCurve : BufferedCurves)
+		{
+			double CurveOutputMin = BufferedCurve->GetValueMin(), CurveOutputMax = BufferedCurve->GetValueMax();
+
+			ViewToBufferedCurveTransform = CalculateViewToCurveTransform(CurveOutputMin, CurveOutputMax, ValueOffset);
+
+			TArray<TTuple<double, double>> CurveSpaceInterpolatingPoints;
+			FCurveEditorScreenSpace CurveSpace = GetViewSpace().ToCurveSpace(ViewToBufferedCurveTransform);
+
+			BufferedCurve->DrawCurve(*CurveEditor, CurveSpace, CurveSpaceInterpolatingPoints);
+
+			TArray<FVector2D> ScreenSpaceInterpolatingPoints;
+			for (TTuple<double, double> Point : CurveSpaceInterpolatingPoints)
+			{
+				ScreenSpaceInterpolatingPoints.Add(FVector2D(
+					CurveSpace.SecondsToScreen(Point.Get<0>()),
+					CurveSpace.ValueToScreen(Point.Get<1>())
+				));
+			}
+
+			FSlateDrawElement::MakeLines(
+				OutDrawElements,
+				CurveLayerId,
+				AllottedGeometry.ToPaintGeometry(),
+				ScreenSpaceInterpolatingPoints,
+				DrawEffects,
+				CurveColor,
+				bAntiAliasCurves,
+				BufferedCurveThickness
+			);
+		}
+
+	}
+}
+
 void SCurveEditorViewStacked::Tick(const FGeometry& AllottedGeometry, const double InCurrentTime, const float InDeltaTime)
 {
 	TSharedPtr<FCurveEditor> CurveEditor = WeakCurveEditor.Pin();
@@ -284,14 +368,7 @@ void SCurveEditorViewStacked::Tick(const FGeometry& AllottedGeometry, const doub
 			double CurveOutputMin = 0, CurveOutputMax = 1;
 			Curve->GetValueRange(CurveOutputMin, CurveOutputMax);
 
-			if (CurveOutputMax > CurveOutputMin)
-			{
-				It->Value.ViewToCurveTransform = Concatenate(FVector2D(0.f, ValueOffset), Concatenate(FScale2D(1.f, (CurveOutputMax - CurveOutputMin)), FVector2D(0.f, CurveOutputMin)));
-			}
-			else
-			{
-				It->Value.ViewToCurveTransform = Concatenate(FVector2D(0.f, ValueOffset-0.5), FVector2D(0.f, CurveOutputMin));
-			}
+			It->Value.ViewToCurveTransform = CalculateViewToCurveTransform(CurveOutputMin, CurveOutputMax, ValueOffset);
 		}
 
 		OutputMax = FMath::Max(OutputMin + CurveInfoByID.Num() + ValueSpacePadding*(CurveInfoByID.Num()+1), 1.0);
diff --git a/Engine/Source/Editor/CurveEditor/Private/Views/SCurveEditorViewStacked.h b/Engine/Source/Editor/CurveEditor/Private/Views/SCurveEditorViewStacked.h
--- a/Engine/Source/Editor/CurveEditor/Private/Views/SCurveEditorViewStacked.h
+++ b/Engine/Source/Editor/CurveEditor/Private/Views/SCurveEditorViewStacked.h
@@ -26,4 +26,5 @@ private:
 
 	void DrawViewGrids(const FGeometry& AllottedGeometry, const FSlateRect& MyCullingRect, FSlateWindowElementList& OutDrawElements, int32 BaseLayerId, ESlateDrawEffect DrawEffects) const;
 	void DrawLabels(const FGeometry& AllottedGeometry, const FSlateRect& MyCullingRect, FSlateWindowElementList& OutDrawElements, int32 BaseLayerId, ESlateDrawEffect DrawEffects) const;
+	void DrawBufferedCurves(const FGeometry& AllottedGeometry, const FSlateRect& MyCullingRect, FSlateWindowElementList& OutDrawElements, int32 BaseLayerId, ESlateDrawEffect DrawEffects) const;
 };
diff --git a/Engine/Source/Editor/CurveEditor/Private/Views/SInteractiveCurveEditorView.cpp b/Engine/Source/Editor/CurveEditor/Private/Views/SInteractiveCurveEditorView.cpp
--- a/Engine/Source/Editor/CurveEditor/Private/Views/SInteractiveCurveEditorView.cpp
+++ b/Engine/Source/Editor/CurveEditor/Private/Views/SInteractiveCurveEditorView.cpp
@@ -186,6 +186,7 @@ void SInteractiveCurveEditorView::PaintView(const FPaintArgs& Args, const FGeome
 
 		DrawBackground(AllottedGeometry, OutDrawElements, BaseLayerId, DrawEffects);
 		DrawGridLines(CurveEditor.ToSharedRef(), AllottedGeometry, OutDrawElements, BaseLayerId, DrawEffects);
+		DrawBufferedCurves(CurveEditor.ToSharedRef(), AllottedGeometry, MyCullingRect, OutDrawElements, BaseLayerId, InWidgetStyle, DrawEffects);
 		DrawCurves(CurveEditor.ToSharedRef(), AllottedGeometry, MyCullingRect, OutDrawElements, BaseLayerId, InWidgetStyle, DrawEffects);
 	}
 }
@@ -442,6 +443,45 @@ void SInteractiveCurveEditorView::DrawCurves(TSharedRef<FCurveEditor> CurveEdito
 	}
 }
 
+void SInteractiveCurveEditorView::DrawBufferedCurves(TSharedRef<FCurveEditor> CurveEditor, const FGeometry& AllottedGeometry, const FSlateRect& MyCullingRect, FSlateWindowElementList& OutDrawElements, int32 BaseLayerId, const FWidgetStyle& InWidgetStyle, ESlateDrawEffect DrawEffects) const
+{
+	const float BufferedCurveThickness = 1.f;
+	const bool  bAntiAliasCurves = true;
+	const FLinearColor CurveColor = CurveViewConstants::BufferedCurveColor;
+	const TArray<TUniquePtr<IBufferedCurveModel>>& BufferedCurves = CurveEditor->GetBufferedCurves();
+
+	const int32 CurveLayerId = BaseLayerId + CurveViewConstants::ELayerOffset::Curves;
+
+	// Draw each buffered curve using the view space transform since the curve space for all curves is the same
+	for (const TUniquePtr<IBufferedCurveModel>& BufferedCurve : BufferedCurves)
+	{
+		TArray<TTuple<double, double>> CurveSpaceInterpolatingPoints;
+		FCurveEditorScreenSpace CurveSpace = GetViewSpace();
+
+		BufferedCurve->DrawCurve(*CurveEditor, CurveSpace, CurveSpaceInterpolatingPoints);
+
+		TArray<FVector2D> ScreenSpaceInterpolatingPoints;
+		for (TTuple<double, double> Point : CurveSpaceInterpolatingPoints)
+		{
+			ScreenSpaceInterpolatingPoints.Add(FVector2D(
+				CurveSpace.SecondsToScreen(Point.Get<0>()),
+				CurveSpace.ValueToScreen(Point.Get<1>())
+			));
+		}
+
+		FSlateDrawElement::MakeLines(
+			OutDrawElements,
+			CurveLayerId,
+			AllottedGeometry.ToPaintGeometry(),
+			ScreenSpaceInterpolatingPoints,
+			DrawEffects,
+			CurveColor,
+			bAntiAliasCurves,
+			BufferedCurveThickness
+		);
+	}
+}
+
 void SInteractiveCurveEditorView::Tick(const FGeometry& AllottedGeometry, const double InCurrentTime, const float InDeltaTime)
 {
 	// Cache our Curve Drawing Params. These are used in multiple places so we cache them once each frame.
diff --git a/Engine/Source/Editor/CurveEditor/Public/CurveEditor.h b/Engine/Source/Editor/CurveEditor/Public/CurveEditor.h
--- a/Engine/Source/Editor/CurveEditor/Public/CurveEditor.h
+++ b/Engine/Source/Editor/CurveEditor/Public/CurveEditor.h
@@ -33,6 +33,7 @@ struct FGeometry;
 struct FCurveEditorSnapMetrics;
 class ICurveEditorExtension;
 class ICurveEditorToolExtension;
+class IBufferedCurveModel;
 
 DECLARE_DELEGATE_OneParam(FOnSetBoolean, bool)
 DECLARE_MULTICAST_DELEGATE_OneParam(FOnCurveArrayChanged, FCurveModelID)
@@ -50,18 +51,6 @@ public:
 	 */
 	FCurveEditorSelection Selection;
 
-	
-protected:
-	struct FBufferedCurve
-	{
-		/** Cached Positions from original curve. */
-		TArray<FKeyPosition>	KeyPositions;
-		/** Cached Attributes from original curve. */
-		TArray<FKeyAttributes>	KeyAttributes;
-		/** The original intention of the type of curve. Can be empty. Used to try and match when pasting multiple curves. */
-		FString					IntentionName;
-	};
-
 public:
 
 	/** Attribute used to retrieve the current input snap rate (also used for display) */
@@ -257,6 +246,8 @@ public:
 	bool ApplyBufferedCurves(const TSet<FCurveModelID>& InCurvesToApplyTo);
 	/** Return the number of stored Buffered Curves. */
 	int32 GetNumBufferedCurves() const { return BufferedCurves.Num(); }
+	/** Return the array of buffered curves */
+	const TArray<TUniquePtr<IBufferedCurveModel>>& GetBufferedCurves() const { return BufferedCurves; }
 	// ~FCurveEditor
 
 	// FEditorUndoClient
@@ -418,7 +409,7 @@ protected:
 	/**
 	*	Apply a specific buffered curve to a specific target curve.
 	*/
-	void ApplyBufferedCurveToTarget(const FBufferedCurve& BufferedCurve, FCurveModel* TargetCurve);
+	void ApplyBufferedCurveToTarget(const IBufferedCurveModel* BufferedCurve, FCurveModel* TargetCurve);
 
 
 protected:
@@ -466,10 +457,8 @@ protected:
 	/** Track which axis UI movements should be snapped to (where applicable) based on limitations imposed by the UI. */
 	FCurveEditorAxisSnap AxisSnapMetrics;
 
-
-
 	/** Buffered Curves. When a curve is buffered it is copied and the new copy is uniquely owned by the Curve Editor. */
-	TArray<FBufferedCurve> BufferedCurves;
+	TArray<TUniquePtr<IBufferedCurveModel>> BufferedCurves;
 
 	/** A serial number that is incremented any time the currently active set of curves are changed */
 	uint32 ActiveCurvesSerialNumber;
diff --git a/Engine/Source/Editor/CurveEditor/Public/CurveModel.h b/Engine/Source/Editor/CurveEditor/Public/CurveModel.h
--- a/Engine/Source/Editor/CurveEditor/Public/CurveModel.h
+++ b/Engine/Source/Editor/CurveEditor/Public/CurveModel.h
@@ -10,6 +10,8 @@
 #include "CurveEditorTypes.h"
 #include "Math/TransformCalculus2D.h"
 
+#include "IBufferedCurveModel.h"
+
 struct FKeyHandle;
 struct FKeyDrawInfo;
 struct FCurveDrawParams;
@@ -202,6 +204,26 @@ public:
 	virtual void CreateKeyProxies(TArrayView<const FKeyHandle> InKeyHandles, TArrayView<UObject*> OutObjects)
 	{}
 
+	/**
+	 * Creates a copy of this curve, stored in a minimal buffered curve object.
+	 * Buffered curves are used to cache the positions and attributes of a curve's keys. After creation, a buffered curve 
+	 * can be applied to any curve to set it to its saved state. Each curve must implement its own buffered curve which 
+	 * inherits IBufferedCurve and implements the DrawCurve method in order for it to be drawn on screen.
+	 * Optionally implemented
+	 */
+	virtual TUniquePtr<IBufferedCurveModel> CreateBufferedCurveCopy() const
+	{
+		return nullptr;
+	}
+
+	/** 
+	 * Returns whether the curve model should be edited or not
+	 */
+	virtual bool IsReadOnly() const
+	{
+		return false;
+	}
+
 	/**
 	 * Helper function for assigning a the same attributes to a number of keys
 	 */
@@ -285,9 +307,9 @@ public:
 	/**
 	 * Retrieve this curve's color
 	 */
-	FORCEINLINE const FLinearColor& GetColor() const
+	FORCEINLINE FLinearColor GetColor() const
 	{
-		return Color;
+		return IsReadOnly() ? Color.Desaturate(.6f) : Color;
 	}
 
 	/**
diff --git a/Engine/Source/Editor/CurveEditor/Public/IBufferedCurveModel.h b/Engine/Source/Editor/CurveEditor/Public/IBufferedCurveModel.h
new file mode 100644
--- /dev/null
+++ b/Engine/Source/Editor/CurveEditor/Public/IBufferedCurveModel.h
@@ -0,0 +1,64 @@
+// Copyright 1998-2019 Epic Games, Inc. All Rights Reserved.
+#pragma once
+
+#include "Containers/Array.h"
+#include "Templates/Tuple.h"
+#include "CurveDataAbstraction.h"
+
+struct FCurveEditorScreenSpace;
+
+class FCurveEditor;
+class FCurveModel;
+
+/**
+ * Represents a buffered curve which can be applied to a standard curve model
+ */
+class IBufferedCurveModel
+{
+public:
+	IBufferedCurveModel(TArray<FKeyPosition>&& InKeyPositions, TArray<FKeyAttributes>&& InKeyAttributes, const FString& InIntentionName, const double InValueMin, const double InValueMax) 
+	: KeyPositions(MoveTemp(InKeyPositions))
+	, KeyAttributes(MoveTemp(InKeyAttributes))
+	, IntentionName(InIntentionName)
+	, ValueMin(InValueMin)
+	, ValueMax(InValueMax)
+	{}
+
+	virtual ~IBufferedCurveModel() = default;
+
+	/** draws the curve into an array of (input, output) pairs with a given screen space */
+	virtual void DrawCurve(const FCurveEditor& CurveEditor, const FCurveEditorScreenSpace& ScreenSpace, TArray<TTuple<double, double>>& InterpolatingPoints) const = 0;
+	
+	/**
+	 * Retrieve all key positions stored in this buffered curve
+	 *
+	 * @param OutKeyPositions         Array to receive key positions, one per index of InKeys
+	 */
+	FORCEINLINE void GetKeyPositions(TArray<FKeyPosition>& OutKeyPositions) const { OutKeyPositions = KeyPositions; }
+	
+	/**
+	 * Retrieve all key attributes stored in this buffered curve
+	 * 
+	 * @param OutAttributes         Array to receive key attributes, one per index of InKeys
+	 */
+	FORCEINLINE void GetKeyAttributes(TArray<FKeyAttributes>& OutKeyAttributes) const { OutKeyAttributes = KeyAttributes; }
+
+	/**
+	 * Returns the intention name for this buffered curve. This is an internal name used to try to match 
+	 * different curves with each other. See the CurveModel's GetIntentionName method for more detail on
+	 * what an intention name is
+	 */
+	FORCEINLINE FString GetIntentionName() const { return IntentionName; }
+
+	/** returns the lowest output value in curve space for this buffered curve */
+	FORCEINLINE double GetValueMin() const { return ValueMin; }
+
+	/** returns the highest output value in curve space for this buffered curve */
+	FORCEINLINE double GetValueMax() const { return ValueMax; }
+
+protected:
+	TArray<FKeyPosition> KeyPositions;
+	TArray<FKeyAttributes> KeyAttributes;
+	FString IntentionName;
+	double ValueMin, ValueMax;
+};
\ No newline at end of file
diff --git a/Engine/Source/Editor/CurveEditor/Public/RichCurveEditorModel.h b/Engine/Source/Editor/CurveEditor/Public/RichCurveEditorModel.h
--- a/Engine/Source/Editor/CurveEditor/Public/RichCurveEditorModel.h
+++ b/Engine/Source/Editor/CurveEditor/Public/RichCurveEditorModel.h
@@ -6,6 +6,7 @@
 #include "UObject/WeakObjectPtr.h"
 #include "UObject/WeakObjectPtrTemplates.h"
 #include "CurveModel.h"
+#include "IBufferedCurveModel.h"
 
 struct FRichCurve;
 
@@ -40,6 +41,8 @@ public:
 
 	virtual void CreateKeyProxies(TArrayView<const FKeyHandle> InKeyHandles, TArrayView<UObject*> OutObjects) override;
 
+	virtual TUniquePtr<IBufferedCurveModel> CreateBufferedCurveCopy() const override;
+
 private:
 
 	FRichCurve* RichCurve;
diff --git a/Engine/Source/Editor/CurveEditor/Public/Views/SInteractiveCurveEditorView.h b/Engine/Source/Editor/CurveEditor/Public/Views/SInteractiveCurveEditorView.h
--- a/Engine/Source/Editor/CurveEditor/Public/Views/SInteractiveCurveEditorView.h
+++ b/Engine/Source/Editor/CurveEditor/Public/Views/SInteractiveCurveEditorView.h
@@ -22,6 +22,8 @@ namespace CurveViewConstants
 	constexpr float CurveLabelOffsetX = 15.f;
 	constexpr float CurveLabelOffsetY = 10.f;
 
+	constexpr FLinearColor BufferedCurveColor = FLinearColor(.4f, .4f, .4f);
+
 	/**
 	 * Pre-defined layer offsets for specific curve view elements. Fixed values are used to decouple draw order and layering
 	 * Some elements deliberately leave some spare layers as a buffer for slight tweaks to layering within that element
@@ -99,6 +101,7 @@ protected:
 	void DrawBackground(const FGeometry& AllottedGeometry, FSlateWindowElementList& OutDrawElements, int32 BaseLayerId, ESlateDrawEffect DrawEffects) const;
 	void DrawGridLines(TSharedRef<FCurveEditor> CurveEditor, const FGeometry& AllottedGeometry, FSlateWindowElementList& OutDrawElements, int32 BaseLayerId, ESlateDrawEffect DrawEffects) const;
 	void DrawCurves(TSharedRef<FCurveEditor> CurveEditor, const FGeometry& AllottedGeometry, const FSlateRect& MyCullingRect, FSlateWindowElementList& OutDrawElements, int32 BaseLayerId, const FWidgetStyle& InWidgetStyle, ESlateDrawEffect DrawEffects) const;
+	void DrawBufferedCurves(TSharedRef<FCurveEditor> CurveEditor, const FGeometry& AllottedGeometry, const FSlateRect& MyCullingRect, FSlateWindowElementList& OutDrawElements, int32 BaseLayerId, const FWidgetStyle& InWidgetStyle, ESlateDrawEffect DrawEffects) const;
 
 	FSlateColor GetCurveCaptionColor() const;
 	FText GetCurveCaption() const;
diff --git a/Engine/Source/Editor/MovieSceneTools/Private/Channels/FloatChannelCurveModel.cpp b/Engine/Source/Editor/MovieSceneTools/Private/Channels/FloatChannelCurveModel.cpp
--- a/Engine/Source/Editor/MovieSceneTools/Private/Channels/FloatChannelCurveModel.cpp
+++ b/Engine/Source/Editor/MovieSceneTools/Private/Channels/FloatChannelCurveModel.cpp
@@ -16,6 +16,43 @@
 #include "BuiltInChannelEditors.h"
 #include "SequencerChannelTraits.h"
 #include "ISequencer.h"
+#include "Channels/MovieSceneChannelProxy.h"
+
+/**
+ * Buffered curve implementation for a float channel curve model, stores a copy of the float channel in order to draw itself.
+ */
+class FFloatChannelBufferedCurveModel : public IBufferedCurveModel
+{
+public:
+	/** Create a copy of the float channel while keeping the reference to the section */
+	FFloatChannelBufferedCurveModel(const FMovieSceneFloatChannel* InMovieSceneFloatChannel, TWeakObjectPtr<UMovieSceneSection> InWeakSection,
+		TArray<FKeyPosition>&& InKeyPositions, TArray<FKeyAttributes>&& InKeyAttributes, const FString& InIntentionName, const double InValueMin, const double InValueMax)
+		: IBufferedCurveModel(MoveTemp(InKeyPositions), MoveTemp(InKeyAttributes), InIntentionName, InValueMin, InValueMax)
+		, Channel(*InMovieSceneFloatChannel)
+		, WeakSection(InWeakSection)
+	{}
+
+	virtual void DrawCurve(const FCurveEditor& InCurveEditor, const FCurveEditorScreenSpace& InScreenSpace, TArray<TTuple<double, double>>& OutInterpolatingPoints) const override
+	{
+		UMovieSceneSection* Section = WeakSection.Get();
+
+		if (Section)
+		{
+			FFrameRate TickResolution = Section->GetTypedOuter<UMovieScene>()->GetTickResolution();
+
+			const double StartTimeSeconds = InScreenSpace.GetInputMin();
+			const double EndTimeSeconds = InScreenSpace.GetInputMax();
+			const double TimeThreshold = FMath::Max(0.0001, 1.0 / InScreenSpace.PixelsPerInput());
+			const double ValueThreshold = FMath::Max(0.0001, 1.0 / InScreenSpace.PixelsPerOutput());
+
+			Channel.PopulateCurvePoints(StartTimeSeconds, EndTimeSeconds, TimeThreshold, ValueThreshold, TickResolution, OutInterpolatingPoints);
+		}
+	}
+
+private:
+	FMovieSceneFloatChannel Channel;
+	TWeakObjectPtr<UMovieSceneSection> WeakSection;
+};
 
 FFloatChannelCurveModel::FFloatChannelCurveModel(TMovieSceneChannelHandle<FMovieSceneFloatChannel> InChannel, UMovieSceneSection* OwningSection, TWeakPtr<ISequencer> InWeakSequencer)
 {
@@ -50,7 +87,7 @@ void FFloatChannelCurveModel::AddKeys(TArrayView<const FKeyPosition> InKeyPositi
 
 	FMovieSceneFloatChannel* Channel = ChannelHandle.Get();
 	UMovieSceneSection*      Section = WeakSection.Get();
-	if (Channel && Section)
+	if (Channel && Section && !IsReadOnly())
 	{
 		Section->Modify();
 
@@ -109,7 +146,7 @@ void FFloatChannelCurveModel::RemoveKeys(TArrayView<const FKeyHandle> InKeys)
 {
 	FMovieSceneFloatChannel* Channel = ChannelHandle.Get();
 	UMovieSceneSection*      Section = WeakSection.Get();
-	if (Channel && Section)
+	if (Channel && Section && !IsReadOnly())
 	{
 		Section->Modify();
 
@@ -262,7 +299,7 @@ void FFloatChannelCurveModel::SetKeyPositions(TArrayView<const FKeyHandle> InKey
 	FMovieSceneFloatChannel* Channel = ChannelHandle.Get();
 	UMovieSceneSection*      Section = WeakSection.Get();
 
-	if (Channel && Section)
+	if (Channel && Section && !IsReadOnly())
 	{
 		Section->MarkAsChanged();
 
@@ -340,7 +377,7 @@ void FFloatChannelCurveModel::SetKeyAttributes(TArrayView<const FKeyHandle> InKe
 {
 	FMovieSceneFloatChannel* Channel = ChannelHandle.Get();
 	UMovieSceneSection*      Section = WeakSection.Get();
-	if (Channel && Section)
+	if (Channel && Section && !IsReadOnly())
 	{
 		bool bAutoSetTangents = false;
 		Section->MarkAsChanged();
@@ -488,7 +525,7 @@ void FFloatChannelCurveModel::SetCurveAttributes(const FCurveAttributes& InCurve
 {
 	FMovieSceneFloatChannel* Channel = ChannelHandle.Get();
 	UMovieSceneSection*      Section = WeakSection.Get();
-	if (Channel && Section)
+	if (Channel && Section && !IsReadOnly())
 	{
 		Section->MarkAsChanged();
 
@@ -515,6 +552,43 @@ void FFloatChannelCurveModel::CreateKeyProxies(TArrayView<const FKeyHandle> InKe
 	}
 }
 
+TUniquePtr<IBufferedCurveModel> FFloatChannelCurveModel::CreateBufferedCurveCopy() const
+{
+	FMovieSceneFloatChannel* Channel = ChannelHandle.Get();
+	if (Channel)
+	{
+		TArray<FKeyHandle> TargetKeyHandles;
+		TMovieSceneChannelData<FMovieSceneFloatValue> ChannelData = Channel->GetData();
+
+		TRange<FFrameNumber> TotalRange = ChannelData.GetTotalRange();
+		ChannelData.GetKeys(TotalRange, nullptr, &TargetKeyHandles);
+
+		TArray<FKeyPosition> KeyPositions;
+		KeyPositions.SetNumUninitialized(GetNumKeys());
+		TArray<FKeyAttributes> KeyAttributes;
+		KeyAttributes.SetNumUninitialized(GetNumKeys());
+		GetKeyPositions(TargetKeyHandles, KeyPositions);
+		GetKeyAttributes(TargetKeyHandles, KeyAttributes);
+
+		double ValueMin = 0.f, ValueMax = 1.f;
+		GetValueRange(ValueMin, ValueMax);
+
+		return MakeUnique<FFloatChannelBufferedCurveModel>(Channel, WeakSection, MoveTemp(KeyPositions), MoveTemp(KeyAttributes), GetIntentionName(), ValueMin, ValueMax);
+	}
+	return nullptr;
+}
+
+bool FFloatChannelCurveModel::IsReadOnly() const
+{
+	UMovieSceneSection* Section = WeakSection.Get();
+	if (Section)
+	{
+		return Section->IsReadOnly();
+	}
+
+	return false;
+}
+
 void FFloatChannelCurveModel::GetTimeRange(double& MinTime, double& MaxTime) const
 {
 	FMovieSceneFloatChannel* Channel = ChannelHandle.Get();
diff --git a/Engine/Source/Editor/MovieSceneTools/Private/Channels/FloatChannelCurveModel.h b/Engine/Source/Editor/MovieSceneTools/Private/Channels/FloatChannelCurveModel.h
--- a/Engine/Source/Editor/MovieSceneTools/Private/Channels/FloatChannelCurveModel.h
+++ b/Engine/Source/Editor/MovieSceneTools/Private/Channels/FloatChannelCurveModel.h
@@ -6,6 +6,7 @@
 #include "UObject/WeakObjectPtr.h"
 #include "UObject/WeakObjectPtrTemplates.h"
 #include "CurveModel.h"
+#include "IBufferedCurveModel.h"
 #include "Channels/MovieSceneFloatChannel.h"
 #include "Channels/MovieSceneChannelHandle.h"
 
@@ -45,6 +46,10 @@ public:
 
 	virtual void CreateKeyProxies(TArrayView<const FKeyHandle> InKeyHandles, TArrayView<UObject*> OutObjects) override;
 
+	virtual TUniquePtr<IBufferedCurveModel> CreateBufferedCurveCopy() const override;
+
+	virtual bool IsReadOnly() const override;
+
 private:
 	void FeaturePointMethod(double StartTime, double EndTime, double StartValue, double Mu, int Depth, int MaxDepth, double& MaxV, double& MinVal) const;
 private:
diff --git a/Engine/Source/Editor/UnrealEd/Private/SCurveEditor.cpp b/Engine/Source/Editor/UnrealEd/Private/SCurveEditor.cpp
--- a/Engine/Source/Editor/UnrealEd/Private/SCurveEditor.cpp
+++ b/Engine/Source/Editor/UnrealEd/Private/SCurveEditor.cpp
@@ -38,8 +38,8 @@
 #include "Framework/Notifications/NotificationManager.h"
 
 #define LOCTEXT_NAMESPACE "SCurveEditor"
-
-DEFINE_LOG_CATEGORY_STATIC(LogCurveEditor, Log, All);
+ 
+DEFINE_LOG_CATEGORY(LogCurveEditor);
 
 const static FVector2D	CONST_KeySize		= FVector2D(11,11);
 const static FVector2D	CONST_TangentSize	= FVector2D(7,7);
diff --git a/Engine/Source/Editor/UnrealEd/Public/SCurveEditor.h b/Engine/Source/Editor/UnrealEd/Public/SCurveEditor.h
--- a/Engine/Source/Editor/UnrealEd/Public/SCurveEditor.h
+++ b/Engine/Source/Editor/UnrealEd/Public/SCurveEditor.h
@@ -29,6 +29,8 @@ class UCurveEditorSettings;
 class UCurveFactory;
 enum class ECheckBoxState : uint8;
 
+UNREALED_API DECLARE_LOG_CATEGORY_EXTERN(LogCurveEditor, Log, All);
+
 //////////////////////////////////////////////////////////////////////////
 // FTrackScaleInfo
 
