Curve Editor: Prevent selection of locked keys

diff --git a/Engine/Source/Editor/CurveEditor/Private/CurveEditor.cpp b/Engine/Source/Editor/CurveEditor/Private/CurveEditor.cpp
--- a/Engine/Source/Editor/CurveEditor/Private/CurveEditor.cpp
+++ b/Engine/Source/Editor/CurveEditor/Private/CurveEditor.cpp
@@ -55,6 +55,8 @@ void FCurveEditor::InitCurveEditor(const FCurveEditorInitParams& InInitParams)
 {
 	ICurveEditorModule& CurveEditorModule = FModuleManager::LoadModuleChecked<ICurveEditorModule>("CurveEditor");
 
+	Selection = FCurveEditorSelection(SharedThis(this));
+
 	// Editor Extensions can be registered in the Curve Editor module. To allow users to derive from FCurveEditor
 	// we have to manually reach out to the module and get a list of extensions to create an instance of them.
 	// If none of your extensions are showing up, it's because you forgot to call this function after construction
diff --git a/Engine/Source/Editor/CurveEditor/Private/CurveEditorSelection.cpp b/Engine/Source/Editor/CurveEditor/Private/CurveEditorSelection.cpp
--- a/Engine/Source/Editor/CurveEditor/Private/CurveEditorSelection.cpp
+++ b/Engine/Source/Editor/CurveEditor/Private/CurveEditorSelection.cpp
@@ -3,6 +3,7 @@
 #include "CurveEditorSelection.h"
 #include "CurveEditorTypes.h"
 #include "Algo/BinarySearch.h"
+#include "CurveEditor.h"
 
 FCurveEditorSelection::FCurveEditorSelection()
 {
@@ -10,6 +11,13 @@ FCurveEditorSelection::FCurveEditorSelection()
 	SerialNumber = 0;
 }
 
+FCurveEditorSelection::FCurveEditorSelection(TWeakPtr<FCurveEditor> InWeakCurveEditor)
+{
+	WeakCurveEditor = InWeakCurveEditor;
+	SelectionType = ECurvePointType::Key;
+	SerialNumber = 0;
+}
+
 void FCurveEditorSelection::Clear()
 {
 	SelectionType = ECurvePointType::Key;
@@ -56,7 +64,8 @@ void FCurveEditorSelection::Add(FCurveModelID CurveID, ECurvePointType PointType
 
 void FCurveEditorSelection::Add(FCurveModelID CurveID, ECurvePointType PointType, TArrayView<const FKeyHandle> Keys)
 {
-	if (Keys.Num() > 0)
+	TSharedPtr<FCurveEditor> CurveEditor = WeakCurveEditor.Pin();
+	if (Keys.Num() > 0 && CurveEditor && !CurveEditor->GetCurves()[CurveID]->IsReadOnly())
 	{
 		ChangeSelectionPointType(PointType);
 
diff --git a/Engine/Source/Editor/CurveEditor/Public/CurveEditorSelection.h b/Engine/Source/Editor/CurveEditor/Public/CurveEditorSelection.h
--- a/Engine/Source/Editor/CurveEditor/Public/CurveEditorSelection.h
+++ b/Engine/Source/Editor/CurveEditor/Public/CurveEditorSelection.h
@@ -11,6 +11,7 @@
 
 struct FCurvePointHandle;
 
+class FCurveEditor;
 
 /**
  * A set of key handles implemented as a sorted array for transparent passing to TArrayView<> APIs.
@@ -66,6 +67,12 @@ struct CURVEEDITOR_API FCurveEditorSelection
 	 */
 	FCurveEditorSelection();
 
+	/**
+	 * Constructor which takes a reference to the curve editor, 
+	 * which is used to find if a model is read only
+	 */
+	FCurveEditorSelection(TWeakPtr<FCurveEditor> InWeakCurveEditor);
+
 	/**
 	 * Retrieve the current type of selection
 	 */
@@ -175,6 +182,9 @@ public:
 
 private:
 
+	/** Weak reference to the curve editor to check whether keys are locked or not */
+	TWeakPtr<FCurveEditor> WeakCurveEditor;
+
 	/** A serial number that increments every time a change is made to the selection */
 	uint32 SerialNumber;
 
diff --git a/Engine/Source/Editor/MovieSceneTools/Private/Channels/FloatChannelCurveModel.cpp b/Engine/Source/Editor/MovieSceneTools/Private/Channels/FloatChannelCurveModel.cpp
--- a/Engine/Source/Editor/MovieSceneTools/Private/Channels/FloatChannelCurveModel.cpp
+++ b/Engine/Source/Editor/MovieSceneTools/Private/Channels/FloatChannelCurveModel.cpp
@@ -266,6 +266,11 @@ void FFloatChannelCurveModel::GetKeyDrawInfo(ECurvePointType PointType, const FK
 			OutDrawInfo.Tint = FLinearColor::White;
 			break;
 		}
+
+		if (IsReadOnly())
+		{
+			OutDrawInfo.Tint = OutDrawInfo.Tint * 0.5f;
+		}
 	}
 }
 
