Curve Editor: Prevent options panel from showing after tool is switched

diff --git a/Engine/Source/Editor/CurveEditor/Private/SCurveEditorToolProperties.cpp b/Engine/Source/Editor/CurveEditor/Private/SCurveEditorToolProperties.cpp
--- a/Engine/Source/Editor/CurveEditor/Private/SCurveEditorToolProperties.cpp
+++ b/Engine/Source/Editor/CurveEditor/Private/SCurveEditorToolProperties.cpp
@@ -64,6 +64,11 @@ void SCurveEditorToolProperties::RebuildProperties()
 			ToolOptions = CurveEditor->GetToolExtensions()[ToolId]->GetToolOptions();
 			DetailsView->SetStructureData(ToolOptions);
 		}
+		else
+		{
+			ToolOptions = nullptr;
+			DetailsView->SetStructureData(nullptr);
+		}
 	}
 }
 
