Implement copy/paste support in the curve editor

diff --git a/Engine/Source/Editor/CurveEditor/CurveEditor.Build.cs b/Engine/Source/Editor/CurveEditor/CurveEditor.Build.cs
--- a/Engine/Source/Editor/CurveEditor/CurveEditor.Build.cs
+++ b/Engine/Source/Editor/CurveEditor/CurveEditor.Build.cs
@@ -8,7 +8,8 @@ public class CurveEditor : ModuleRules
 	{
 		PrivateDependencyModuleNames.AddRange(
 			new string[] {
-				"Core",
+                "ApplicationCore",
+                "Core",
 				"CoreUObject",
 				"Engine",
 				"EditorStyle",
diff --git a/Engine/Source/Editor/CurveEditor/Private/CurveEditor.cpp b/Engine/Source/Editor/CurveEditor/Private/CurveEditor.cpp
--- a/Engine/Source/Editor/CurveEditor/Private/CurveEditor.cpp
+++ b/Engine/Source/Editor/CurveEditor/Private/CurveEditor.cpp
@@ -6,6 +6,7 @@
 #include "CurveEditorCommands.h"
 #include "CurveEditorSettings.h"
 #include "CurveDrawInfo.h"
+#include "CurveEditorCopyBuffer.h"
 #include "Framework/Commands/GenericCommands.h"
 #include "Framework/Commands/UICommandList.h"
 #include "Editor.h"
@@ -22,6 +23,10 @@
 #include "Framework/Notifications/NotificationManager.h"
 #include "Widgets/Notifications/SNotificationList.h"
 #include "Runtime/Core/Public/Algo/Transform.h"
+#include "UnrealExporter.h"
+#include "Exporters/Exporter.h"
+#include "Factories.h"
+#include "HAL/PlatformApplicationMisc.h"
 #include "SCurveEditor.h" // for access to LogCurveEditor
 
 #define LOCTEXT_NAMESPACE "CurveEditor"
@@ -240,6 +245,10 @@ void FCurveEditor::BindCommands()
 	CommandList->MapAction(FGenericCommands::Get().Redo,   FExecuteAction::CreateLambda([]{ GEditor->RedoTransaction(); }));
 	CommandList->MapAction(FGenericCommands::Get().Delete, FExecuteAction::CreateSP(this, &FCurveEditor::DeleteSelection));
 
+	CommandList->MapAction(FGenericCommands::Get().Cut, FExecuteAction::CreateSP(this, &FCurveEditor::CutSelection));
+	CommandList->MapAction(FGenericCommands::Get().Copy, FExecuteAction::CreateSP(this, &FCurveEditor::CopySelection));
+	CommandList->MapAction(FGenericCommands::Get().Paste, FExecuteAction::CreateSP(this, &FCurveEditor::PasteKeys));
+
 	CommandList->MapAction(FCurveEditorCommands::Get().ZoomToFit, FExecuteAction::CreateSP(this, &FCurveEditor::ZoomToFit, EAxisList::All));
 
 	CommandList->MapAction(FCurveEditorCommands::Get().StepToNextKey, FExecuteAction::CreateSP(this, &FCurveEditor::StepToNextKey));
@@ -711,6 +720,312 @@ void FCurveEditor::ConstructXGridLines(TArray<float>& MajorGridLines, TArray<flo
 	}
 }
 
+void FCurveEditor::CutSelection()
+{
+	FScopedTransaction Transaction(LOCTEXT("CutKeys", "Cut Keys"));
+
+	CopySelection();
+	DeleteSelection();
+}
+
+void FCurveEditor::GetChildCurveModelIDs(const FCurveEditorTreeItemID TreeItemID, TSet<FCurveModelID>& OutCurveModelIDs) const
+{
+	const FCurveEditorTreeItem& TreeItem = GetTreeItem(TreeItemID);
+	for (const FCurveModelID CurveModelID : TreeItem.GetCurves())
+	{
+		OutCurveModelIDs.Add(CurveModelID);
+	}
+
+	for (const FCurveEditorTreeItemID& ChildTreeItem : TreeItem.GetChildren())
+	{
+		GetChildCurveModelIDs(ChildTreeItem, OutCurveModelIDs);
+	}
+}
+
+void FCurveEditor::CopySelection() const
+{
+	FStringOutputDevice Archive;
+	const FExportObjectInnerContext Context;
+
+	TOptional<double> KeyOffset;
+
+	UCurveEditorCopyBuffer* CopyableBuffer = NewObject<UCurveEditorCopyBuffer>(GetTransientPackage(), UCurveEditorCopyBuffer::StaticClass(), NAME_None, RF_Transient);
+
+	if (Selection.Count() > 0)
+	{
+		for (const TTuple<FCurveModelID, FKeyHandleSet>& Pair : Selection.GetAll())
+		{
+			if (FCurveModel* Curve = FindCurve(Pair.Key))
+			{
+				int32 NumKeys = Pair.Value.Num();
+
+				if (NumKeys > 0)
+				{
+					UCurveEditorCopyableCurveKeys *CopyableCurveKeys = NewObject<UCurveEditorCopyableCurveKeys>(CopyableBuffer, UCurveEditorCopyableCurveKeys::StaticClass(), NAME_None, RF_Transient);
+
+					CopyableCurveKeys->ShortDisplayName = Curve->GetShortDisplayName().ToString();
+					CopyableCurveKeys->LongDisplayName = Curve->GetLongDisplayName().ToString();
+					CopyableCurveKeys->IntentionName = Curve->GetIntentionName();
+					CopyableCurveKeys->KeyPositions.SetNum(NumKeys, false);
+					CopyableCurveKeys->KeyAttributes.SetNum(NumKeys, false);
+
+					TArrayView<const FKeyHandle> KeyHandles = Pair.Value.AsArray();
+
+					Curve->GetKeyPositions(KeyHandles, CopyableCurveKeys->KeyPositions);
+					Curve->GetKeyAttributes(KeyHandles, CopyableCurveKeys->KeyAttributes);
+
+					for (int KeyIndex = 0; KeyIndex < CopyableCurveKeys->KeyPositions.Num(); ++KeyIndex)
+					{
+						if (!KeyOffset.IsSet() || CopyableCurveKeys->KeyPositions[KeyIndex].InputValue < KeyOffset.GetValue())
+						{
+							KeyOffset = CopyableCurveKeys->KeyPositions[KeyIndex].InputValue;
+						}
+					}
+
+					CopyableBuffer->Curves.Add(CopyableCurveKeys);
+				}
+			}
+		}
+	}
+	else
+	{
+		TSet<FCurveModelID> CurveModelIDs;
+
+		for (const TTuple<FCurveEditorTreeItemID, ECurveEditorTreeSelectionState> Pair : GetTreeSelection())
+		{
+			if (Pair.Value == ECurveEditorTreeSelectionState::Explicit)
+			{
+				GetChildCurveModelIDs(Pair.Key, CurveModelIDs);
+			}
+		}
+
+		for(const FCurveModelID CurveModelID : CurveModelIDs)
+		{
+			if (FCurveModel* Curve = FindCurve(CurveModelID))
+			{
+				TUniquePtr<IBufferedCurveModel> CurveModelCopy = Curve->CreateBufferedCurveCopy();
+				if (CurveModelCopy)
+				{
+					TArray<FKeyPosition> KeyPositions;
+					CurveModelCopy->GetKeyPositions(KeyPositions);
+					if (KeyPositions.Num() > 0)
+					{
+						UCurveEditorCopyableCurveKeys *CopyableCurveKeys = NewObject<UCurveEditorCopyableCurveKeys>(CopyableBuffer, UCurveEditorCopyableCurveKeys::StaticClass(), NAME_None, RF_Transient);
+
+						CopyableCurveKeys->ShortDisplayName = Curve->GetShortDisplayName().ToString();
+						CopyableCurveKeys->LongDisplayName = Curve->GetLongDisplayName().ToString();
+						CopyableCurveKeys->IntentionName = Curve->GetIntentionName();
+
+						CopyableCurveKeys->KeyPositions = KeyPositions;
+						CurveModelCopy->GetKeyAttributes(CopyableCurveKeys->KeyAttributes);
+
+						CopyableBuffer->Curves.Add(CopyableCurveKeys);
+					}
+				}
+			}
+		}
+
+		// When copying entire curve objects we want absolute positions, so reset the detected offset
+		KeyOffset.Reset();
+	}
+
+	if (KeyOffset.IsSet())
+	{
+		for (UCurveEditorCopyableCurveKeys* Curve : CopyableBuffer->Curves)
+		{
+			for (int Index = 0; Index < Curve->KeyPositions.Num(); ++Index)
+			{
+				Curve->KeyPositions[Index].InputValue -= KeyOffset.GetValue();
+			}
+		}
+
+		CopyableBuffer->TimeOffset = KeyOffset.GetValue();
+	}
+	else
+	{
+		CopyableBuffer->bAbsolutePosition = true;
+	}
+
+
+	UExporter::ExportToOutputDevice(&Context, CopyableBuffer, nullptr, Archive, TEXT("copy"), 0, PPF_ExportsNotFullyQualified | PPF_Copy | PPF_Delimited, false, CopyableBuffer);
+	FPlatformApplicationMisc::ClipboardCopy(*Archive);
+}
+
+class FCurveEditorCopyableCurveKeysObjectTextFactory : public FCustomizableTextObjectFactory
+{
+public:
+	FCurveEditorCopyableCurveKeysObjectTextFactory()
+		: FCustomizableTextObjectFactory(GWarn)
+	{
+	}
+
+	// FCustomizableTextObjectFactory implementation
+	virtual bool CanCreateClass(UClass* InObjectClass, bool& bOmitSubObjs) const override
+	{
+		if (InObjectClass->IsChildOf(UCurveEditorCopyBuffer::StaticClass()))
+		{
+			return true;
+		}
+		return false;
+	}
+
+
+	virtual void ProcessConstructedObject(UObject* NewObject) override
+	{
+		check(NewObject);
+
+		NewCopyBuffers.Add(Cast<UCurveEditorCopyBuffer>(NewObject));
+	}
+
+public:
+	TArray<UCurveEditorCopyBuffer*> NewCopyBuffers;
+};
+
+
+bool FCurveEditor::CanPaste(const FString& TextToImport) const
+{
+	FCurveEditorCopyableCurveKeysObjectTextFactory CopyableCurveKeysFactory;
+	if (CopyableCurveKeysFactory.CanCreateObjectsFromText(TextToImport))
+	{
+		return true;
+	}
+	return false;
+}
+
+void FCurveEditor::ImportCopyBufferFromText(const FString& TextToImport, /*out*/ TArray<UCurveEditorCopyBuffer*>& ImportedCopyBuffers) const
+{
+	UPackage* TempPackage = NewObject<UPackage>(nullptr, TEXT("/Engine/Editor/CurveEditor/Transient"), RF_Transient);
+	TempPackage->AddToRoot();
+
+	// Turn the text buffer into objects
+	FCurveEditorCopyableCurveKeysObjectTextFactory Factory;
+	Factory.ProcessBuffer(TempPackage, RF_Transactional, TextToImport);
+
+	ImportedCopyBuffers = Factory.NewCopyBuffers;
+
+	// Remove the temp package from the root now that it has served its purpose
+	TempPackage->RemoveFromRoot();
+}
+
+void FCurveEditor::PasteKeys()
+{
+	// Grab the text to paste from the clipboard
+	FString TextToImport;
+	FPlatformApplicationMisc::ClipboardPaste(TextToImport);
+
+	TArray<UCurveEditorCopyBuffer*> ImportedCopyBuffers;
+	ImportCopyBufferFromText(TextToImport, ImportedCopyBuffers);
+
+	if (ImportedCopyBuffers.Num() == 0)
+	{
+		return;
+	}
+
+	bool bMultipleRootsSelected = false;
+	FCurveEditorTreeItemID LastRootItem;
+
+	for (const TTuple<FCurveEditorTreeItemID, ECurveEditorTreeSelectionState> Pair : GetTreeSelection())
+	{
+		FCurveEditorTreeItem* TreeItem = &GetTreeItem(Pair.Key);
+
+		FCurveEditorTreeItemID ParentId = TreeItem->GetParentID();
+		while (ParentId.IsValid())
+		{
+			TreeItem = &GetTreeItem(ParentId);
+			ParentId = TreeItem->GetParentID();
+		}
+
+		if (!LastRootItem.IsValid())
+		{
+			LastRootItem = TreeItem->GetID();
+		}
+		else if (TreeItem->GetID() != LastRootItem)
+		{
+			bMultipleRootsSelected = true;
+			break;
+		}
+	}
+
+	FScopedTransaction Transaction(LOCTEXT("PasteKeys", "Paste Keys"));
+
+	Selection.Clear();
+
+	// We don't expect/want multiple copy buffers, but the way serialization works it's a possibile edge case,
+	// so we'll try to handle it sanely and treat each one as an individual block to paste.
+	for (UCurveEditorCopyBuffer* CopyBuffer : ImportedCopyBuffers)
+	{
+		bool bUseLongDisplayName = bMultipleRootsSelected;
+		if (!bUseLongDisplayName)
+		{
+			TOptional<FString> LastRootName;
+			for (UCurveEditorCopyableCurveKeys* CopyableCurveKeys : CopyBuffer->Curves)
+			{
+				FString RootName;
+				CopyableCurveKeys->LongDisplayName.Split(".", &RootName, nullptr);
+				if (!LastRootName.IsSet())
+				{
+					LastRootName = RootName;
+				}
+				else if (!LastRootName.GetValue().Equals(RootName))
+				{
+					bUseLongDisplayName = true;
+					break;
+				}
+			}
+		}
+
+		double TimeOffset = 0.0f;
+		bool bApplyOffset = !CopyBuffer->bAbsolutePosition;
+
+		if (bApplyOffset)
+		{
+			if (WeakTimeSliderController.IsValid())
+			{
+				FFrameRate TickResolution = WeakTimeSliderController.Pin()->GetTickResolution();
+
+				TimeOffset = WeakTimeSliderController.Pin()->GetScrubPosition() / TickResolution;
+			}
+			else
+			{
+				TimeOffset = CopyBuffer->TimeOffset;
+			}
+		}
+
+		for (const TTuple<FCurveEditorTreeItemID, ECurveEditorTreeSelectionState> Pair : GetTreeSelection())
+		{
+			FCurveEditorTreeItem& TreeItem = GetTreeItem(Pair.Key);
+			for (const FCurveModelID CurveModelID : TreeItem.GetCurves())
+			{
+				FCurveModel* Curve = FindCurve(CurveModelID);
+				const FString CurveLongDisplayName = Curve->GetLongDisplayName().ToString();
+				const FString CurveIntentionName = Curve->GetIntentionName();
+
+				for (UCurveEditorCopyableCurveKeys* CopyableCurveKeys : CopyBuffer->Curves)
+				{
+					if ((!bUseLongDisplayName && CurveIntentionName.Equals(CopyableCurveKeys->IntentionName))
+						|| (bUseLongDisplayName && CurveLongDisplayName.Equals(CopyableCurveKeys->LongDisplayName)))
+					{
+						for (int32 Index = 0; Index < CopyableCurveKeys->KeyPositions.Num(); ++Index)
+						{
+							FKeyPosition KeyPosition = CopyableCurveKeys->KeyPositions[Index];
+							if (bApplyOffset)
+							{
+								KeyPosition.InputValue += TimeOffset;
+							}
+
+							TOptional<FKeyHandle> KeyHandle = Curve->AddKey(KeyPosition, CopyableCurveKeys->KeyAttributes[Index]);
+							if (KeyHandle.IsSet())
+							{
+								Selection.Add(FCurvePointHandle(CurveModelID, ECurvePointType::Key, KeyHandle.GetValue()));
+							}
+						}
+					}
+				}
+			}
+		}
+	}
+}
+
 void FCurveEditor::DeleteSelection()
 {
 	FScopedTransaction Transaction(LOCTEXT("DeleteKeys", "Delete Keys"));
diff --git a/Engine/Source/Editor/CurveEditor/Private/CurveEditorCopyBuffer.h b/Engine/Source/Editor/CurveEditor/Private/CurveEditorCopyBuffer.h
new file mode 100644
--- /dev/null
+++ b/Engine/Source/Editor/CurveEditor/Private/CurveEditorCopyBuffer.h
@@ -0,0 +1,49 @@
+// Copyright 1998-2019 Epic Games, Inc. All Rights Reserved.
+
+#pragma once
+#include "CoreMinimal.h"
+#include "UObject/ObjectMacros.h"
+#include "CurveDataAbstraction.h"
+#include "CurveEditorCopyBuffer.generated.h"
+
+UCLASS()
+class UCurveEditorCopyableCurveKeys : public UObject
+{
+public:
+	GENERATED_BODY()
+
+	UPROPERTY()
+	TArray<FKeyPosition> KeyPositions;
+
+	UPROPERTY()
+	TArray<FKeyAttributes> KeyAttributes;
+
+	/** This curve's short display name. Used in situations where other mechanisms provide enough context about what the curve is (such as "X") */
+	UPROPERTY()
+	FString ShortDisplayName;
+
+	/** This curve's long display name. Used in situations where the UI doesn't provide enough context about what the curve is otherwise (such as "Floor.Transform.X") */
+	UPROPERTY()
+	FString LongDisplayName;
+
+	/** This curve's intention (such as Transform.X or Scale.X). Used internally to match up curves when saving/restoring curves between different objects. */
+	UPROPERTY()
+	FString IntentionName;
+};
+
+UCLASS()
+class UCurveEditorCopyBuffer : public UObject
+{
+public:
+	GENERATED_BODY()
+
+	UPROPERTY(Instanced)
+	TArray<UCurveEditorCopyableCurveKeys*> Curves;
+
+	UPROPERTY()
+	double TimeOffset;
+	
+	UPROPERTY()
+	bool bAbsolutePosition;
+};
+
diff --git a/Engine/Source/Editor/CurveEditor/Public/CurveDataAbstraction.h b/Engine/Source/Editor/CurveEditor/Public/CurveDataAbstraction.h
--- a/Engine/Source/Editor/CurveEditor/Public/CurveDataAbstraction.h
+++ b/Engine/Source/Editor/CurveEditor/Public/CurveDataAbstraction.h
@@ -5,14 +5,18 @@
 #include "CoreTypes.h"
 #include "CurveEditorTypes.h"
 #include "Curves/RichCurve.h"
+#include "CurveDataAbstraction.generated.h"
 
 struct FSlateBrush;
 
 /**
  * Generic key position information for a key on a curve
  */
+USTRUCT()
 struct FKeyPosition
 {
+	GENERATED_BODY()
+
 	FKeyPosition()
 		: InputValue(0.0), OutputValue(0.0)
 	{}
@@ -22,16 +26,22 @@ struct FKeyPosition
 	{}
 
 	/** The key's input (x-axis) position (i.e. it's time) */
+	UPROPERTY()
 	double InputValue;
+
 	/** The key's output (t-axis) position (i.e. it's value) */
+	UPROPERTY()
 	double OutputValue;
 };
 
 /**
  * Extended attributes that the curve editor understands
  */
+USTRUCT()
 struct FKeyAttributes
 {
+	GENERATED_BODY()
+		
 	FKeyAttributes()
 	{
 		bHasArriveTangent = 0;
@@ -133,33 +143,47 @@ struct FKeyAttributes
 private:
 
 	/** True if this key supports entry tangents */
+	UPROPERTY()
 	uint8 bHasArriveTangent : 1;
 	/** True if this key supports exit tangents */
+	UPROPERTY()
 	uint8 bHasLeaveTangent : 1;
 	/** True if this key supports interpolation modes */
+	UPROPERTY()
 	uint8 bHasInterpMode : 1;
 	/** True if this key supports tangent modes */
+	UPROPERTY()
 	uint8 bHasTangentMode : 1;
 	/** True if this key supports tangent modes */
+	UPROPERTY()
 	uint8 bHasTangentWeightMode : 1;
 	/** True if this key supports entry tangents weights*/
+	UPROPERTY()
 	uint8 bHasArriveTangentWeight : 1;
 	/** True if this key supports exit tangents weights*/
+	UPROPERTY()
 	uint8 bHasLeaveTangentWeight : 1;
 
 	/** This key's entry tangent, if bHasArriveTangent */
+	UPROPERTY()
 	float ArriveTangent;
 	/** This key's exit tangent, if bHasLeaveTangent */
+	UPROPERTY()
 	float LeaveTangent;
 	/** This key's interpolation mode, if bHasInterpMode */
-	ERichCurveInterpMode InterpMode;
+	UPROPERTY()
+	TEnumAsByte<ERichCurveInterpMode> InterpMode;
 	/** This key's tangent mode, if bHasTangentMode */
-	ERichCurveTangentMode TangentMode;
+	UPROPERTY()
+	TEnumAsByte<ERichCurveTangentMode> TangentMode;
 	/** This key's tangent weight mode, if bHasTangentWeightMode */
-	ERichCurveTangentWeightMode TangentWeightMode;
+	UPROPERTY()
+	TEnumAsByte<ERichCurveTangentWeightMode> TangentWeightMode;
 	/** This key's entry tangent weight, if bHasArriveTangentWeight */
+	UPROPERTY()
 	float ArriveTangentWeight;
 	/** This key's exit tangent weight, if bHasLeaveTangentWeight */
+	UPROPERTY()
 	float LeaveTangentWeight;
 
 };
diff --git a/Engine/Source/Editor/CurveEditor/Public/CurveEditor.h b/Engine/Source/Editor/CurveEditor/Public/CurveEditor.h
--- a/Engine/Source/Editor/CurveEditor/Public/CurveEditor.h
+++ b/Engine/Source/Editor/CurveEditor/Public/CurveEditor.h
@@ -28,6 +28,7 @@ class UCurveEditorSettings;
 class FUICommandList;
 class FCurveEditor;
 class SCurveEditorPanel;
+class UCurveEditorCopyBuffer;
 
 struct FGeometry;
 struct FCurveEditorSnapMetrics;
@@ -347,6 +348,32 @@ public:
 
 public:
 
+	/**
+	 * Cut the currently selected keys
+	 */
+	void CutSelection();
+	
+	/**
+	 * Copy the currently selected keys
+	 */
+	void CopySelection() const;
+
+	/**
+	 * Returns whether the current clipboard contains objects which CurveEditor can paste
+	 */
+	bool CanPaste(const FString& TextToImport) const;
+
+protected:
+	void ImportCopyBufferFromText(const FString& TextToImport, /*out*/ TArray<UCurveEditorCopyBuffer*>& ImportedCopyBuffers) const;
+
+	void GetChildCurveModelIDs(const FCurveEditorTreeItemID TreeItemID, TSet<FCurveModelID>& CurveModelIDs) const;
+
+public:
+	/**
+	 * Paste keys
+	 */
+	void PasteKeys();
+
 	/**
 	 * Delete the currently selected keys
 	 */
