Curve Editor: Fix crash when stepping backwards in the curve editor

diff --git a/Engine/Source/Editor/CurveEditor/Private/CurveEditor.cpp b/Engine/Source/Editor/CurveEditor/Private/CurveEditor.cpp
--- a/Engine/Source/Editor/CurveEditor/Private/CurveEditor.cpp
+++ b/Engine/Source/Editor/CurveEditor/Private/CurveEditor.cpp
@@ -605,6 +605,11 @@ void FCurveEditor::StepForward()
 
 void FCurveEditor::StepBackward()
 {
+	if (!WeakTimeSliderController.IsValid())
+	{
+		return;
+	}
+
 	FFrameRate TickResolution = WeakTimeSliderController.Pin()->GetTickResolution();
 	FFrameRate DisplayRate = WeakTimeSliderController.Pin()->GetDisplayRate();
 
