Curve Editor: Correctly size arrays when created buffered rich curves, to prevent crash

diff --git a/Engine/Source/Editor/CurveEditor/Private/RichCurveEditorModel.cpp b/Engine/Source/Editor/CurveEditor/Private/RichCurveEditorModel.cpp
--- a/Engine/Source/Editor/CurveEditor/Private/RichCurveEditorModel.cpp
+++ b/Engine/Source/Editor/CurveEditor/Private/RichCurveEditorModel.cpp
@@ -545,7 +545,9 @@ TUniquePtr<IBufferedCurveModel> FRichCurveEditorModel::CreateBufferedCurveCopy()
 		}
 
 		TArray<FKeyPosition> KeyPositions;
+		KeyPositions.SetNumUninitialized(TargetKeyHandles.Num());
 		TArray<FKeyAttributes> KeyAttributes;
+		KeyAttributes.SetNumUninitialized(TargetKeyHandles.Num());
 		GetKeyPositions(TargetKeyHandles, KeyPositions);
 		GetKeyAttributes(TargetKeyHandles, KeyAttributes);
 
