Curve editing tools now specify interactive vs. non-interactive changes
#rnx

diff --git a/Engine/Plugins/Editor/CurveEditorTools/Source/CurveEditorTools/Private/Tools/CurveEditorRetimeTool.cpp b/Engine/Plugins/Editor/CurveEditorTools/Source/CurveEditorTools/Private/Tools/CurveEditorRetimeTool.cpp
--- a/Engine/Plugins/Editor/CurveEditorTools/Source/CurveEditorTools/Private/Tools/CurveEditorRetimeTool.cpp
+++ b/Engine/Plugins/Editor/CurveEditorTools/Source/CurveEditorTools/Private/Tools/CurveEditorRetimeTool.cpp
@@ -481,6 +481,7 @@ void FCurveEditorRetimeTool::OnDragStart()
 			// Get all Key Positions for those Handles.
 			ChannelData.FrameNumbers.SetNumUninitialized(ChannelData.Handles.Num());
 			CurveModel->GetKeyPositions(ChannelData.Handles, ChannelData.FrameNumbers);
+			ChannelData.LastDraggedFrameNumbers = ChannelData.FrameNumbers;
 		}
 
 		// Store our new data for use in OnDrag
@@ -594,7 +595,8 @@ void FCurveEditorRetimeTool::OnDrag(const FVector2D& InStartPosition, const FVec
 				NewPosition.InputValue = View->IsTimeSnapEnabled() ? CurveEditor->GetCurveSnapMetrics(ChannelData.CurveID).SnapInputSeconds(NewKeyValue) : NewKeyValue;
 
 				// Update our Curve Model with our updated data.
-				Curve->SetKeyPositions(MakeArrayView(&ChannelData.Handles[KeyIndex], 1), MakeArrayView(&NewPosition, 1));
+				Curve->SetKeyPositions(MakeArrayView(&ChannelData.Handles[KeyIndex], 1), MakeArrayView(&NewPosition, 1), EPropertyChangeType::Interactive);
+				ChannelData.LastDraggedFrameNumbers[KeyIndex] = NewPosition;
 			}
 		}
 	}
@@ -605,6 +607,22 @@ void FCurveEditorRetimeTool::OnDrag(const FVector2D& InStartPosition, const FVec
 
 void FCurveEditorRetimeTool::OnDragEnd()
 {
+	TSharedPtr<FCurveEditor> CurveEditor = WeakCurveEditor.Pin();
+	// For each curve we're affecting.
+	for (FPreDragChannelData& ChannelData : PreDragCurveData->CurveChannels)
+	{
+		for (int32 KeyIndex = 0; KeyIndex < ChannelData.Handles.Num(); KeyIndex++)
+		{
+			if (CurveEditor)
+			{
+				FCurveModel* Curve = CurveEditor->FindCurve(ChannelData.CurveID);
+				if (Curve)
+				{
+					Curve->SetKeyPositions(MakeArrayView(&ChannelData.Handles[KeyIndex], 1), MakeArrayView(&ChannelData.LastDraggedFrameNumbers[KeyIndex], 1), EPropertyChangeType::ValueSet);
+				}
+			}
+		}
+	}
 	// This finalizes the transaction
 	ActiveTransaction.Reset();
 	PreDragCurveData.Reset();
diff --git a/Engine/Plugins/Editor/CurveEditorTools/Source/CurveEditorTools/Private/Tools/CurveEditorRetimeTool.h b/Engine/Plugins/Editor/CurveEditorTools/Source/CurveEditorTools/Private/Tools/CurveEditorRetimeTool.h
--- a/Engine/Plugins/Editor/CurveEditorTools/Source/CurveEditorTools/Private/Tools/CurveEditorRetimeTool.h
+++ b/Engine/Plugins/Editor/CurveEditorTools/Source/CurveEditorTools/Private/Tools/CurveEditorRetimeTool.h
@@ -115,7 +115,8 @@ private:
 		TArray<FKeyHandle> Handles;
 		/** Array of all the above handle's times, one per index of Handles */
 		TArray<FKeyPosition> FrameNumbers;
-		
+		/** Array of all the above handle's last dragged times, one per index of Handles */
+		TArray<FKeyPosition> LastDraggedFrameNumbers;
 		FCurveModelID CurveID;
 	};
 
diff --git a/Engine/Plugins/Editor/CurveEditorTools/Source/CurveEditorTools/Private/Tools/CurveEditorTransformTool.cpp b/Engine/Plugins/Editor/CurveEditorTools/Source/CurveEditorTools/Private/Tools/CurveEditorTransformTool.cpp
--- a/Engine/Plugins/Editor/CurveEditorTools/Source/CurveEditorTools/Private/Tools/CurveEditorTransformTool.cpp
+++ b/Engine/Plugins/Editor/CurveEditorTools/Source/CurveEditorTools/Private/Tools/CurveEditorTransformTool.cpp
@@ -788,6 +788,7 @@ void FCurveEditorTransformTool::OnDragStart()
 
 			KeyData.StartKeyPositions.SetNumZeroed(KeyData.Handles.Num());
 			Curve->GetKeyPositions(KeyData.Handles, KeyData.StartKeyPositions);
+			KeyData.LastDraggedKeyPositions = KeyData.StartKeyPositions;
 			for (const FKeyPosition& KeyPosition : KeyData.StartKeyPositions)
 			{
 				if (bMinMaxIsSet)
@@ -985,7 +986,7 @@ void FCurveEditorTransformTool::OnDrag(const FPointerEvent& InMouseEvent, const
 				TransformWidget.Position = TransformWidget.StartPosition + MouseDelta;
 			}
 
-			for (const FKeyData& KeyData : KeysByCurve)
+			for (FKeyData& KeyData : KeysByCurve)
 			{
 				const SCurveEditorView* View = CurveEditor->FindFirstInteractiveView(KeyData.CurveID);
 				if (!View)
@@ -1015,7 +1016,8 @@ void FCurveEditorTransformTool::OnDrag(const FPointerEvent& InMouseEvent, const
 					NewKeyPositionScratch.Add(StartPosition);
 				}
 
-				CurveModel->SetKeyPositions(KeyData.Handles, NewKeyPositionScratch);
+				CurveModel->SetKeyPositions(KeyData.Handles, NewKeyPositionScratch, EPropertyChangeType::Interactive);
+				KeyData.LastDraggedKeyPositions = NewKeyPositionScratch;
 			}
 		}
 		else if (TransformWidget.SelectedAnchorFlags != ECurveEditorAnchorFlags::None)
@@ -1070,6 +1072,14 @@ void FCurveEditorTransformTool::OnDragEnd()
 		return;
 	}
 
+	for (const FKeyData& KeyData : KeysByCurve)
+	{
+		if (FCurveModel* Curve = CurveEditor->FindCurve(KeyData.CurveID))
+		{
+			Curve->SetKeyPositions(KeyData.Handles, KeyData.LastDraggedKeyPositions, EPropertyChangeType::ValueSet);
+		}
+	}
+
 	CurveEditor->SuppressBoundTransformUpdates(false);
 
 	// This finalizes the transaction
@@ -1166,7 +1176,7 @@ void FCurveEditorTransformTool::ScaleFrom(const FVector2D& InPanelSpaceCenter, c
 			NewKeyPositionScratch.Add(StartPosition);
 		}
 
-		CurveModel->SetKeyPositions(KeyData.Handles, NewKeyPositionScratch);
+		CurveModel->SetKeyPositions(KeyData.Handles, NewKeyPositionScratch, EPropertyChangeType::Interactive);
 	}
 }
 
diff --git a/Engine/Plugins/Editor/CurveEditorTools/Source/CurveEditorTools/Private/Tools/CurveEditorTransformTool.h b/Engine/Plugins/Editor/CurveEditorTools/Source/CurveEditorTools/Private/Tools/CurveEditorTransformTool.h
--- a/Engine/Plugins/Editor/CurveEditorTools/Source/CurveEditorTools/Private/Tools/CurveEditorTransformTool.h
+++ b/Engine/Plugins/Editor/CurveEditorTools/Source/CurveEditorTools/Private/Tools/CurveEditorTransformTool.h
@@ -175,6 +175,8 @@ private:
 		TArray<FKeyHandle> Handles;
 		/** The extended key info for each of the above handles */
 		TArray<FKeyPosition> StartKeyPositions;
+		/** Used in OnEndDrag to send final key updates */
+		TArray<FKeyPosition> LastDraggedKeyPositions;
 	};
 
 	/** Key dragging data stored per-curve */
diff --git a/Engine/Source/Editor/CurveEditor/Private/CurveModel.cpp b/Engine/Source/Editor/CurveEditor/Private/CurveModel.cpp
--- a/Engine/Source/Editor/CurveEditor/Private/CurveModel.cpp
+++ b/Engine/Source/Editor/CurveEditor/Private/CurveModel.cpp
@@ -4,7 +4,7 @@
 #include "Containers/Array.h"
 #include "CurveDataAbstraction.h"
 
-void FCurveModel::SetKeyAttributes(TArrayView<const FKeyHandle> InKeys, const FKeyAttributes& InKeyAttributes)
+void FCurveModel::SetKeyAttributes(TArrayView<const FKeyHandle> InKeys, const FKeyAttributes& InKeyAttributes, EPropertyChangeType::Type ChangeType)
 {
 	TArray<FKeyAttributes> ExpandedAttributes;
 	ExpandedAttributes.Reserve(InKeys.Num());
diff --git a/Engine/Source/Editor/CurveEditor/Private/DragOperations/CurveEditorDragOperation_MoveKeys.cpp b/Engine/Source/Editor/CurveEditor/Private/DragOperations/CurveEditorDragOperation_MoveKeys.cpp
--- a/Engine/Source/Editor/CurveEditor/Private/DragOperations/CurveEditorDragOperation_MoveKeys.cpp
+++ b/Engine/Source/Editor/CurveEditor/Private/DragOperations/CurveEditorDragOperation_MoveKeys.cpp
@@ -34,6 +34,7 @@ void FCurveEditorDragOperation_MoveKeys::OnBeginDrag(FVector2D InitialPosition,
 
 			KeyData.StartKeyPositions.SetNumZeroed(KeyData.Handles.Num());
 			Curve->GetKeyPositions(KeyData.Handles, KeyData.StartKeyPositions);
+			KeyData.LastDraggedKeyPositions = KeyData.StartKeyPositions;
 		}
 	}
 
@@ -45,7 +46,7 @@ void FCurveEditorDragOperation_MoveKeys::OnDrag(FVector2D InitialPosition, FVect
 	TArray<FKeyPosition> NewKeyPositionScratch;
 	FVector2D MousePosition = CurveEditor->GetAxisSnap().GetSnappedPosition(InitialPosition, CurrentPosition, MouseEvent, SnappingState);
 
-	for (const FKeyData& KeyData : KeysByCurve)
+	for (FKeyData& KeyData : KeysByCurve)
 	{
 		const SCurveEditorView* View = CurveEditor->FindFirstInteractiveView(KeyData.CurveID);
 		if (!View)
@@ -80,7 +81,8 @@ void FCurveEditorDragOperation_MoveKeys::OnDrag(FVector2D InitialPosition, FVect
 			NewKeyPositionScratch.Add(StartPosition);
 		}
 
-		Curve->SetKeyPositions(KeyData.Handles, NewKeyPositionScratch);
+		Curve->SetKeyPositions(KeyData.Handles, NewKeyPositionScratch, EPropertyChangeType::Interactive);
+		KeyData.LastDraggedKeyPositions = NewKeyPositionScratch;
 	}
 }
 
@@ -92,7 +94,7 @@ void FCurveEditorDragOperation_MoveKeys::OnCancelDrag()
 	{
 		if (FCurveModel* Curve = CurveEditor->FindCurve(KeyData.CurveID))
 		{
-			Curve->SetKeyPositions(KeyData.Handles, KeyData.StartKeyPositions);
+			Curve->SetKeyPositions(KeyData.Handles, KeyData.StartKeyPositions, EPropertyChangeType::ValueSet);
 		}
 	}
 
@@ -102,5 +104,14 @@ void FCurveEditorDragOperation_MoveKeys::OnCancelDrag()
 void FCurveEditorDragOperation_MoveKeys::OnEndDrag(FVector2D InitialPosition, FVector2D CurrentPosition, const FPointerEvent& MouseEvent)
 {
 	ICurveEditorKeyDragOperation::OnEndDrag(InitialPosition, CurrentPosition, MouseEvent);
+
+	for (const FKeyData& KeyData : KeysByCurve)
+	{
+		if (FCurveModel* Curve = CurveEditor->FindCurve(KeyData.CurveID))
+		{
+			Curve->SetKeyPositions(KeyData.Handles, KeyData.LastDraggedKeyPositions, EPropertyChangeType::ValueSet);
+		}
+	}
+
 	CurveEditor->SuppressBoundTransformUpdates(false);
 }
diff --git a/Engine/Source/Editor/CurveEditor/Private/DragOperations/CurveEditorDragOperation_MoveKeys.h b/Engine/Source/Editor/CurveEditor/Private/DragOperations/CurveEditorDragOperation_MoveKeys.h
--- a/Engine/Source/Editor/CurveEditor/Private/DragOperations/CurveEditorDragOperation_MoveKeys.h
+++ b/Engine/Source/Editor/CurveEditor/Private/DragOperations/CurveEditorDragOperation_MoveKeys.h
@@ -39,6 +39,8 @@ private:
 		TArray<FKeyHandle> Handles;
 		/** The extended key info for each of the above handles */
 		TArray<FKeyPosition> StartKeyPositions;
+		/** Used in OnEndDrag to send final key updates */
+		TArray<FKeyPosition> LastDraggedKeyPositions;
 	};
 
 	/** Key dragging data stored per-curve */
diff --git a/Engine/Source/Editor/CurveEditor/Private/RichCurveEditorModel.cpp b/Engine/Source/Editor/CurveEditor/Private/RichCurveEditorModel.cpp
--- a/Engine/Source/Editor/CurveEditor/Private/RichCurveEditorModel.cpp
+++ b/Engine/Source/Editor/CurveEditor/Private/RichCurveEditorModel.cpp
@@ -149,7 +149,7 @@ void FRichCurveEditorModel::AddKeys(TArrayView<const FKeyPosition> InKeyPosition
 
 		// We reuse SetKeyAttributes here as there is complex logic determining which parts of the attributes are valid to pass on.
 		// For now we need to duplicate the new key handle array due to API mismatch. This will auto-calculate tangents if required.
-		SetKeyAttributes(NewKeyHandles, InKeyAttributes);
+		SetKeyAttributes(NewKeyHandles, InKeyAttributes, EPropertyChangeType::ValueSet);
 	}
 }
 
@@ -274,7 +274,7 @@ void FRichCurveEditorModel::GetKeyPositions(TArrayView<const FKeyHandle> InKeys,
 	}
 }
 
-void FRichCurveEditorModel::SetKeyPositions(TArrayView<const FKeyHandle> InKeys, TArrayView<const FKeyPosition> InKeyPositions)
+void FRichCurveEditorModel::SetKeyPositions(TArrayView<const FKeyHandle> InKeys, TArrayView<const FKeyPosition> InKeyPositions, EPropertyChangeType::Type ChangeType)
 {
 	if (IsReadOnly())
 	{
@@ -296,7 +296,7 @@ void FRichCurveEditorModel::SetKeyPositions(TArrayView<const FKeyHandle> InKeys,
 			}
 		}
 		RichCurve->AutoSetTangents();
-		FPropertyChangedEvent PropertyChangeStruct(nullptr, EPropertyChangeType::ValueSet);
+		FPropertyChangedEvent PropertyChangeStruct(nullptr, ChangeType);
 		Owner->PostEditChangeProperty(PropertyChangeStruct);
 	}
 }
@@ -350,7 +350,7 @@ void FRichCurveEditorModel::GetKeyAttributes(TArrayView<const FKeyHandle> InKeys
 	}
 }
 
-void FRichCurveEditorModel::SetKeyAttributes(TArrayView<const FKeyHandle> InKeys, TArrayView<const FKeyAttributes> InAttributes)
+void FRichCurveEditorModel::SetKeyAttributes(TArrayView<const FKeyHandle> InKeys, TArrayView<const FKeyAttributes> InAttributes, EPropertyChangeType::Type ChangeType)
 {
 	if (IsReadOnly())
 	{
@@ -488,7 +488,7 @@ void FRichCurveEditorModel::SetKeyAttributes(TArrayView<const FKeyHandle> InKeys
 			RichCurve->AutoSetTangents();
 		}
 
-		FPropertyChangedEvent PropertyChangeStruct(nullptr, EPropertyChangeType::ValueSet);
+		FPropertyChangedEvent PropertyChangeStruct(nullptr, ChangeType);
 		Owner->PostEditChangeProperty(PropertyChangeStruct);
 	}
 }
diff --git a/Engine/Source/Editor/CurveEditor/Public/CurveModel.h b/Engine/Source/Editor/CurveEditor/Public/CurveModel.h
--- a/Engine/Source/Editor/CurveEditor/Public/CurveModel.h
+++ b/Engine/Source/Editor/CurveEditor/Public/CurveModel.h
@@ -106,7 +106,7 @@ public:
 	 * @param InKeys                 Array of key handles to set positions for
 	 * @param InKeyPositions         Array of desired key positions to be applied to each of the corresponding key handles
 	 */
-	virtual void SetKeyPositions(TArrayView<const FKeyHandle> InKeys, TArrayView<const FKeyPosition> InKeyPositions) = 0;
+	virtual void SetKeyPositions(TArrayView<const FKeyHandle> InKeys, TArrayView<const FKeyPosition> InKeyPositions, EPropertyChangeType::Type ChangeType = EPropertyChangeType::Unspecified) = 0;
 
 	/**
 	 * Populate the specified draw info structure with data describing how to draw the specified point type
@@ -169,7 +169,7 @@ public:
 	 * @param InKeys                 Array of key handles to set attributes for
 	 * @param InAttributes           Array of desired key attributes to be applied to each of the corresponding key handles
 	 */
-	virtual void SetKeyAttributes(TArrayView<const FKeyHandle> InKeys, TArrayView<const FKeyAttributes> InAttributes)
+	virtual void SetKeyAttributes(TArrayView<const FKeyHandle> InKeys, TArrayView<const FKeyAttributes> InAttributes, EPropertyChangeType::Type ChangeType = EPropertyChangeType::Unspecified)
 	{}
 
 	/**
@@ -236,7 +236,7 @@ public:
 	/**
 	 * Helper function for assigning a the same attributes to a number of keys
 	 */
-	void SetKeyAttributes(TArrayView<const FKeyHandle> InKeys, const FKeyAttributes& InAttributes);
+	void SetKeyAttributes(TArrayView<const FKeyHandle> InKeys, const FKeyAttributes& InAttributes, EPropertyChangeType::Type ChangeType = EPropertyChangeType::Unspecified);
 
 	/**
 	 * Helper function for adding a single key to this curve
diff --git a/Engine/Source/Editor/CurveEditor/Public/RichCurveEditorModel.h b/Engine/Source/Editor/CurveEditor/Public/RichCurveEditorModel.h
--- a/Engine/Source/Editor/CurveEditor/Public/RichCurveEditorModel.h
+++ b/Engine/Source/Editor/CurveEditor/Public/RichCurveEditorModel.h
@@ -24,10 +24,10 @@ public:
 	virtual void GetKeyDrawInfo(ECurvePointType PointType, const FKeyHandle InKeyHandle, FKeyDrawInfo& OutDrawInfo) const override;
 
 	virtual void GetKeyPositions(TArrayView<const FKeyHandle> InKeys, TArrayView<FKeyPosition> OutKeyPositions) const override;
-	virtual void SetKeyPositions(TArrayView<const FKeyHandle> InKeys, TArrayView<const FKeyPosition> InKeyPositions) override;
+	virtual void SetKeyPositions(TArrayView<const FKeyHandle> InKeys, TArrayView<const FKeyPosition> InKeyPositions, EPropertyChangeType::Type ChangeType) override;
 
 	virtual void GetKeyAttributes(TArrayView<const FKeyHandle> InKeys, TArrayView<FKeyAttributes> OutAttributes) const override;
-	virtual void SetKeyAttributes(TArrayView<const FKeyHandle> InKeys, TArrayView<const FKeyAttributes> InAttributes) override;
+	virtual void SetKeyAttributes(TArrayView<const FKeyHandle> InKeys, TArrayView<const FKeyAttributes> InAttributes, EPropertyChangeType::Type ChangeType = EPropertyChangeType::Unspecified) override;
 
 	virtual void GetCurveAttributes(FCurveAttributes& OutCurveAttributes) const override;
 	virtual void SetCurveAttributes(const FCurveAttributes& InCurveAttributes) override;
diff --git a/Engine/Source/Editor/MovieSceneTools/Private/Channels/EventChannelCurveModel.cpp b/Engine/Source/Editor/MovieSceneTools/Private/Channels/EventChannelCurveModel.cpp
--- a/Engine/Source/Editor/MovieSceneTools/Private/Channels/EventChannelCurveModel.cpp
+++ b/Engine/Source/Editor/MovieSceneTools/Private/Channels/EventChannelCurveModel.cpp
@@ -164,7 +164,7 @@ void FEventChannelCurveModel::GetKeyPositions(TArrayView<const FKeyHandle> InKey
 	}
 }
 
-void FEventChannelCurveModel::SetKeyPositions(TArrayView<const FKeyHandle> InKeys, TArrayView<const FKeyPosition> InKeyPositions)
+void FEventChannelCurveModel::SetKeyPositions(TArrayView<const FKeyHandle> InKeys, TArrayView<const FKeyPosition> InKeyPositions, EPropertyChangeType::Type ChangeType)
 {
 	FMovieSceneEventChannel* Channel = ChannelHandle.Get();
 	UMovieSceneSection*      Section = WeakSection.Get();
@@ -194,7 +194,7 @@ void FEventChannelCurveModel::GetKeyAttributes(TArrayView<const FKeyHandle> InKe
 {
 }
 
-void FEventChannelCurveModel::SetKeyAttributes(TArrayView<const FKeyHandle> InKeys, TArrayView<const FKeyAttributes> InAttributes)
+void FEventChannelCurveModel::SetKeyAttributes(TArrayView<const FKeyHandle> InKeys, TArrayView<const FKeyAttributes> InAttributes, EPropertyChangeType::Type ChangeType)
 {
 }
 
diff --git a/Engine/Source/Editor/MovieSceneTools/Private/Channels/EventChannelCurveModel.h b/Engine/Source/Editor/MovieSceneTools/Private/Channels/EventChannelCurveModel.h
--- a/Engine/Source/Editor/MovieSceneTools/Private/Channels/EventChannelCurveModel.h
+++ b/Engine/Source/Editor/MovieSceneTools/Private/Channels/EventChannelCurveModel.h
@@ -28,10 +28,10 @@ public:
 	virtual void GetKeyDrawInfo(ECurvePointType PointType, const FKeyHandle InKeyHandle, FKeyDrawInfo& OutDrawInfo) const override;
 
 	virtual void GetKeyPositions(TArrayView<const FKeyHandle> InKeys, TArrayView<FKeyPosition> OutKeyPositions) const override;
-	virtual void SetKeyPositions(TArrayView<const FKeyHandle> InKeys, TArrayView<const FKeyPosition> InKeyPositions) override;
+	virtual void SetKeyPositions(TArrayView<const FKeyHandle> InKeys, TArrayView<const FKeyPosition> InKeyPositions, EPropertyChangeType::Type ChangeType) override;
 
 	virtual void GetKeyAttributes(TArrayView<const FKeyHandle> InKeys, TArrayView<FKeyAttributes> OutAttributes) const override;
-	virtual void SetKeyAttributes(TArrayView<const FKeyHandle> InKeys, TArrayView<const FKeyAttributes> InAttributes) override;
+	virtual void SetKeyAttributes(TArrayView<const FKeyHandle> InKeys, TArrayView<const FKeyAttributes> InAttributes, EPropertyChangeType::Type ChangeType = EPropertyChangeType::Unspecified) override;
 
 	virtual void GetCurveAttributes(FCurveAttributes& OutCurveAttributes) const override;
 	virtual void SetCurveAttributes(const FCurveAttributes& InCurveAttributes) override;
diff --git a/Engine/Source/Editor/MovieSceneTools/Private/Channels/FloatChannelCurveModel.cpp b/Engine/Source/Editor/MovieSceneTools/Private/Channels/FloatChannelCurveModel.cpp
--- a/Engine/Source/Editor/MovieSceneTools/Private/Channels/FloatChannelCurveModel.cpp
+++ b/Engine/Source/Editor/MovieSceneTools/Private/Channels/FloatChannelCurveModel.cpp
@@ -299,7 +299,7 @@ void FFloatChannelCurveModel::GetKeyPositions(TArrayView<const FKeyHandle> InKey
 	}
 }
 
-void FFloatChannelCurveModel::SetKeyPositions(TArrayView<const FKeyHandle> InKeys, TArrayView<const FKeyPosition> InKeyPositions)
+void FFloatChannelCurveModel::SetKeyPositions(TArrayView<const FKeyHandle> InKeys, TArrayView<const FKeyPosition> InKeyPositions, EPropertyChangeType::Type ChangeType)
 {
 	FMovieSceneFloatChannel* Channel = ChannelHandle.Get();
 	UMovieSceneSection*      Section = WeakSection.Get();
@@ -378,7 +378,7 @@ void FFloatChannelCurveModel::GetKeyAttributes(TArrayView<const FKeyHandle> InKe
 	}
 }
 
-void FFloatChannelCurveModel::SetKeyAttributes(TArrayView<const FKeyHandle> InKeys, TArrayView<const FKeyAttributes> InAttributes)
+void FFloatChannelCurveModel::SetKeyAttributes(TArrayView<const FKeyHandle> InKeys, TArrayView<const FKeyAttributes> InAttributes, EPropertyChangeType::Type ChangeType)
 {
 	FMovieSceneFloatChannel* Channel = ChannelHandle.Get();
 	UMovieSceneSection*      Section = WeakSection.Get();
diff --git a/Engine/Source/Editor/MovieSceneTools/Private/Channels/FloatChannelCurveModel.h b/Engine/Source/Editor/MovieSceneTools/Private/Channels/FloatChannelCurveModel.h
--- a/Engine/Source/Editor/MovieSceneTools/Private/Channels/FloatChannelCurveModel.h
+++ b/Engine/Source/Editor/MovieSceneTools/Private/Channels/FloatChannelCurveModel.h
@@ -29,10 +29,10 @@ public:
 	virtual void GetKeyDrawInfo(ECurvePointType PointType, const FKeyHandle InKeyHandle, FKeyDrawInfo& OutDrawInfo) const override;
 
 	virtual void GetKeyPositions(TArrayView<const FKeyHandle> InKeys, TArrayView<FKeyPosition> OutKeyPositions) const override;
-	virtual void SetKeyPositions(TArrayView<const FKeyHandle> InKeys, TArrayView<const FKeyPosition> InKeyPositions) override;
+	virtual void SetKeyPositions(TArrayView<const FKeyHandle> InKeys, TArrayView<const FKeyPosition> InKeyPositions, EPropertyChangeType::Type ChangeType) override;
 
 	virtual void GetKeyAttributes(TArrayView<const FKeyHandle> InKeys, TArrayView<FKeyAttributes> OutAttributes) const override;
-	virtual void SetKeyAttributes(TArrayView<const FKeyHandle> InKeys, TArrayView<const FKeyAttributes> InAttributes) override;
+	virtual void SetKeyAttributes(TArrayView<const FKeyHandle> InKeys, TArrayView<const FKeyAttributes> InAttributes, EPropertyChangeType::Type ChangeType = EPropertyChangeType::Unspecified) override;
 
 	virtual void GetCurveAttributes(FCurveAttributes& OutCurveAttributes) const override;
 	virtual void SetCurveAttributes(const FCurveAttributes& InCurveAttributes) override;
