COPY from //Dev-Anim to //Dev-Main

diff --git a/Engine/Source/Editor/AnimationModifiers/Public/AnimationBlueprintLibrary.h b/Engine/Source/Editor/AnimationModifiers/Public/AnimationBlueprintLibrary.h
--- a/Engine/Source/Editor/AnimationModifiers/Public/AnimationBlueprintLibrary.h
+++ b/Engine/Source/Editor/AnimationModifiers/Public/AnimationBlueprintLibrary.h
@@ -12,6 +12,8 @@
 struct FRawAnimSequenceTrack;
 class UAnimCompress;
 class USkeleton;
+class UAnimBoneCompressionSettings;
+class UAnimCurveCompressionSettings;
 
 UENUM()
 enum class ESmartNameContainerType : uint8
@@ -66,13 +68,21 @@ public:
 
 	// Compression
 
-	/** Retrieves the Compression Scheme for the given Animation Sequence */
+	/** Retrieves the Bone Compression Settings for the given Animation Sequence */
 	UFUNCTION(BlueprintPure, Category = "AnimationBlueprintLibrary|Compression")
-	static void GetCompressionScheme(const UAnimSequence* AnimationSequence, UAnimCompress*& CompressionScheme);
+	static void GetBoneCompressionSettings(const UAnimSequence* AnimationSequence, UAnimBoneCompressionSettings*& CompressionSettings);
 
-	/** Sets the Compression Scheme for the given Animation Sequence */
+	/** Sets the Bone Compression Settings for the given Animation Sequence */
 	UFUNCTION(BlueprintCallable, Category = "AnimationBlueprintLibrary|Compression")
-	static void SetCompressionScheme(UAnimSequence* AnimationSequence, UAnimCompress* CompressionScheme);
+	static void SetBoneCompressionSettings(UAnimSequence* AnimationSequence, UAnimBoneCompressionSettings* CompressionSettings);
+
+	/** Retrieves the Curve Compression Settings for the given Animation Sequence */
+	UFUNCTION(BlueprintPure, Category = "AnimationBlueprintLibrary|Compression")
+	static void GetCurveCompressionSettings(const UAnimSequence* AnimationSequence, UAnimCurveCompressionSettings*& CompressionSettings);
+
+	/** Sets the Curve Compression Settings for the given Animation Sequence */
+	UFUNCTION(BlueprintCallable, Category = "AnimationBlueprintLibrary|Compression")
+	static void SetCurveCompressionSettings(UAnimSequence* AnimationSequence, UAnimCurveCompressionSettings* CompressionSettings);
 
 	// Additive 
 	/** Retrieves the Additive Animation type for the given Animation Sequence */
diff --git a/Engine/Source/Editor/ClassViewer/Private/SClassViewer.cpp b/Engine/Source/Editor/ClassViewer/Private/SClassViewer.cpp
--- a/Engine/Source/Editor/ClassViewer/Private/SClassViewer.cpp
+++ b/Engine/Source/Editor/ClassViewer/Private/SClassViewer.cpp
@@ -1234,7 +1234,7 @@ bool FClassHierarchy::FindAndRemoveNodeByClassPath(const TSharedPtr< FClassViewe
 	// Search the children recursively, one of them might have the parent.
 	for(int32 ChildClassIndex = 0; ChildClassIndex < InRootNode->GetChildrenList().Num(); ChildClassIndex++)
 	{
-		if(InRootNode->GetChildrenList()[ChildClassIndex]->ClassPath == InClassPath)
+		if(InRootNode->GetChildrenList()[ChildClassIndex]->ClassPath == InClassPath)						   
 		{
 			InRootNode->GetChildrenList().RemoveAt(ChildClassIndex);
 			return true;
@@ -1256,6 +1256,14 @@ void FClassHierarchy::RemoveAsset(const FAssetData& InRemovedAssetData)
 	if (InRemovedAssetData.GetTagValue(FBlueprintTags::GeneratedClassPath, ClassObjectPath))
 	{
 		ClassObjectPath = FPackageName::ExportTextPathToObjectPath(ClassObjectPath);
+
+		if (ClassObjectPath == "None")
+		{
+			// This can happen if the generated class was already deleted prior to 
+			// the notification being sent. Let's try to reconstruct the generated
+			// class name from the object path.
+			ClassObjectPath = InRemovedAssetData.ObjectPath.ToString() + "_C";
+		}
 	}
 
 	if (FindAndRemoveNodeByClassPath(ObjectClassRoot, FName(*ClassObjectPath)))
diff --git a/Engine/Source/Editor/ClothPainter/ClothPainter.Build.cs b/Engine/Source/Editor/ClothPainter/ClothPainter.Build.cs
--- a/Engine/Source/Editor/ClothPainter/ClothPainter.Build.cs
+++ b/Engine/Source/Editor/ClothPainter/ClothPainter.Build.cs
@@ -22,7 +22,6 @@ public class ClothPainter : ModuleRules
                 "Kismet",
                 "AssetTools",
                 "ClassViewer",
-                "Persona",
                 "SkeletalMeshEditor",                
                 "AssetRegistry",
                 "MeshPaint",
diff --git a/Engine/Source/Editor/ContentBrowser/Private/SAssetPicker.cpp b/Engine/Source/Editor/ContentBrowser/Private/SAssetPicker.cpp
--- a/Engine/Source/Editor/ContentBrowser/Private/SAssetPicker.cpp
+++ b/Engine/Source/Editor/ContentBrowser/Private/SAssetPicker.cpp
@@ -298,6 +298,7 @@ void SAssetPicker::Construct( const FArguments& InArgs )
 		.FilterRecursivelyWithBackendFilter( false )
 		.CanShowRealTimeThumbnails( InArgs._AssetPickerConfig.bCanShowRealTimeThumbnails )
 		.CanShowDevelopersFolder( InArgs._AssetPickerConfig.bCanShowDevelopersFolder )
+		.ForceShowEngineContent( InArgs._AssetPickerConfig.bForceShowEngineContent)
 		.PreloadAssetsForContextMenu( InArgs._AssetPickerConfig.bPreloadAssetsForContextMenu )
 		.HighlightedText( HighlightText )
 		.ThumbnailLabel( ThumbnailLabel )
diff --git a/Engine/Source/Editor/ContentBrowser/Private/SAssetView.cpp b/Engine/Source/Editor/ContentBrowser/Private/SAssetView.cpp
--- a/Engine/Source/Editor/ContentBrowser/Private/SAssetView.cpp
+++ b/Engine/Source/Editor/ContentBrowser/Private/SAssetView.cpp
@@ -484,6 +484,7 @@ void SAssetView::Construct( const FArguments& InArgs )
 	bShowPathInColumnView = InArgs._ShowPathInColumnView;
 	bShowTypeInColumnView = InArgs._ShowTypeInColumnView;
 	bSortByPathInColumnView = bShowPathInColumnView & InArgs._SortByPathInColumnView;
+	bForceShowEngineContent = InArgs._ForceShowEngineContent;
 
 	bPendingUpdateThumbnails = false;
 	bShouldNotifyNextAssetSync = true;
@@ -3323,7 +3324,7 @@ void SAssetView::PopulateViewButtonMenu(UToolMenu* Menu)
 			FSlateIcon(),
 			FUIAction(
 				FExecuteAction::CreateSP( this, &SAssetView::ToggleShowEngineContent ),
-				FCanExecuteAction(),
+				FCanExecuteAction::CreateSP( this, &SAssetView::IsToggleShowEngineContentAllowed),
 				FIsActionChecked::CreateSP( this, &SAssetView::IsShowingEngineContent )
 			),
 			EUserInterfaceActionType::ToggleButton
@@ -3567,7 +3568,7 @@ void SAssetView::ToggleShowEngineContent()
 
 bool SAssetView::IsShowingEngineContent() const
 {
-	return GetDefault<UContentBrowserSettings>()->GetDisplayEngineFolder();
+	return bForceShowEngineContent || GetDefault<UContentBrowserSettings>()->GetDisplayEngineFolder();
 }
 
 void SAssetView::ToggleShowDevelopersContent()
@@ -3593,6 +3594,11 @@ bool SAssetView::IsToggleShowDevelopersContentAllowed() const
 	return bCanShowDevelopersFolder;
 }
 
+bool SAssetView::IsToggleShowEngineContentAllowed() const
+{
+	return !bForceShowEngineContent;
+}
+
 bool SAssetView::IsShowingDevelopersContent() const
 {
 	return IsToggleShowDevelopersContentAllowed() && GetDefault<UContentBrowserSettings>()->GetDisplayDevelopersFolder();
diff --git a/Engine/Source/Editor/CurveAssetEditor/Private/CurveAssetEditor.cpp b/Engine/Source/Editor/CurveAssetEditor/Private/CurveAssetEditor.cpp
--- a/Engine/Source/Editor/CurveAssetEditor/Private/CurveAssetEditor.cpp
+++ b/Engine/Source/Editor/CurveAssetEditor/Private/CurveAssetEditor.cpp
@@ -75,7 +75,7 @@ struct FCurveAssetEditorTreeItem : public ICurveEditorTreeItem
 			return;
 		}
 
-		TUniquePtr<FRichCurveEditorModel> NewCurve = MakeUnique<FRichCurveEditorModel>(static_cast<FRichCurve*>(EditInfo.CurveToEdit), CurveOwner.Get());
+		TUniquePtr<FRichCurveEditorModelRaw> NewCurve = MakeUnique<FRichCurveEditorModelRaw>(static_cast<FRichCurve*>(EditInfo.CurveToEdit), CurveOwner.Get());
 		NewCurve->SetShortDisplayName(CurveName);
 		NewCurve->SetColor(CurveColor);
 		OutCurveModels.Add(MoveTemp(NewCurve));
diff --git a/Engine/Source/Editor/CurveEditor/Private/CurveEditor.cpp b/Engine/Source/Editor/CurveEditor/Private/CurveEditor.cpp
--- a/Engine/Source/Editor/CurveEditor/Private/CurveEditor.cpp
+++ b/Engine/Source/Editor/CurveEditor/Private/CurveEditor.cpp
@@ -44,6 +44,7 @@ FCurveEditor::FCurveEditor()
 	: Bounds(new FStaticCurveEditorBounds)
 	, bBoundTransformUpdatesSuppressed(false)
 	, ActiveCurvesSerialNumber(0)
+	, SuspendBroadcastCount(0)
 {
 	Settings = GetMutableDefault<UCurveEditorSettings>();
 	CommandList = MakeShared<FUICommandList>();
@@ -84,7 +85,7 @@ void FCurveEditor::InitCurveEditor(const FCurveEditorInitParams& InInitParams)
 		// We call a delegate and have the delegate create the instance to cover cross-module
 		AddTool(Tools[DelegateIndex].Execute(SharedThis(this)));
 	}
-
+	SuspendBroadcastCount = 0;
 	// Listen to global undo so we can fix up our selection state for keys that no longer exist.
 	GEditor->RegisterForUndo(this);
 }
@@ -99,6 +100,16 @@ TSharedPtr<SCurveEditorPanel> FCurveEditor::GetPanel() const
 	return WeakPanel.Pin();
 }
 
+void FCurveEditor::SetView(TSharedPtr<SCurveEditorView> InView)
+{
+	WeakView = InView;
+}
+
+TSharedPtr<SCurveEditorView> FCurveEditor::GetView() const
+{
+	return WeakView.Pin();
+}
+
 FCurveModel* FCurveEditor::FindCurve(FCurveModelID CurveID) const
 {
 	const TUniquePtr<FCurveModel>* Ptr = CurveData.Find(CurveID);
@@ -121,16 +132,34 @@ FCurveEditorToolID FCurveEditor::AddTool(TUniquePtr<ICurveEditorToolExtension>&&
 FCurveModelID FCurveEditor::AddCurve(TUniquePtr<FCurveModel>&& InCurve)
 {
 	FCurveModelID NewID = FCurveModelID::Unique();
-	CurveData.Add(NewID, MoveTemp(InCurve));
+	FCurveModel *Curve = InCurve.Get();
 
+	CurveData.Add(NewID, MoveTemp(InCurve));
 	++ActiveCurvesSerialNumber;
-
+	if (IsBroadcasting())
+	{
+		OnCurveArrayChanged.Broadcast(Curve, true);
+	}
 	return NewID;
 }
 
+void FCurveEditor::BroadcastCurveChanged(FCurveModel* InCurve)
+{
+	if (IsBroadcasting())
+	{
+		OnCurveArrayChanged.Broadcast(InCurve, true);
+	}
+}
+
 FCurveModelID FCurveEditor::AddCurveForTreeItem(TUniquePtr<FCurveModel>&& InCurve, FCurveEditorTreeItemID TreeItemID)
 {
 	FCurveModelID NewID = FCurveModelID::Unique();
+	FCurveModel *Curve = InCurve.Get();
+
+	if(IsBroadcasting())
+	{
+		OnCurveArrayChanged.Broadcast(InCurve.Get(), true);
+	}
 
 	CurveData.Add(NewID, MoveTemp(InCurve));
 	TreeIDByCurveID.Add(NewID, TreeItemID);
@@ -148,10 +177,36 @@ void FCurveEditor::RemoveCurve(FCurveModelID InCurveID)
 		Panel->RemoveCurveFromViews(InCurveID);
 	}
 
+	if(IsBroadcasting())
+	{
+		OnCurveArrayChanged.Broadcast(FindCurve(InCurveID), false);
+	}
+
+
 	CurveData.Remove(InCurveID);
 	Selection.Remove(InCurveID);
 	PinnedCurves.Remove(InCurveID);
 
+
+	++ActiveCurvesSerialNumber;
+
+}
+
+void FCurveEditor::RemoveAllCurves()
+{
+	TSharedPtr<SCurveEditorPanel> Panel = WeakPanel.Pin();
+	if (Panel.IsValid())
+	{
+		for (TPair<FCurveModelID, TUniquePtr<FCurveModel>>& CurvePair : CurveData)
+		{
+			Panel->RemoveCurveFromViews(CurvePair.Key);
+		}
+	}
+
+	CurveData.Empty();
+	Selection.Clear();
+	PinnedCurves.Empty();
+
 	++ActiveCurvesSerialNumber;
 }
 
@@ -174,11 +229,15 @@ void FCurveEditor::UnpinCurve(FCurveModelID InCurveID)
 
 const SCurveEditorView* FCurveEditor::FindFirstInteractiveView(FCurveModelID InCurveID) const
 {
-	for (auto ViewIt = GetPanel()->FindViews(InCurveID); ViewIt; ++ViewIt)
+	TSharedPtr<SCurveEditorPanel> Panel = WeakPanel.Pin();
+	if (Panel.IsValid())
 	{
-		if (ViewIt.Value()->IsInteractive())
+		for (auto ViewIt = Panel->FindViews(InCurveID); ViewIt; ++ViewIt)
 		{
-			return &ViewIt.Value().Get();
+			if (ViewIt.Value()->IsInteractive())
+			{
+				return &ViewIt.Value().Get();
+			}
 		}
 	}
 	return nullptr;
@@ -216,6 +275,21 @@ void FCurveEditor::RemoveTreeItem(FCurveEditorTreeItemID ItemID)
 	++ActiveCurvesSerialNumber;
 }
 
+void FCurveEditor::RemoveAllTreeItems()
+{
+	TArray<FCurveEditorTreeItemID> RootItems = Tree.GetRootItems();
+	for(FCurveEditorTreeItemID ItemID : RootItems)
+	{
+		Tree.RemoveItem(ItemID, this);
+	}
+	++ActiveCurvesSerialNumber;
+}
+
+void FCurveEditor::SetTreeSelection(TArray<FCurveEditorTreeItemID>&& TreeItems)
+{
+	Tree.SetDirectSelection(MoveTemp(TreeItems), this);
+}
+
 ECurveEditorTreeSelectionState FCurveEditor::GetTreeSelectionState(FCurveEditorTreeItemID InTreeItemID) const
 {
 	return Tree.GetSelectionState(InTreeItemID);
@@ -434,10 +508,28 @@ void FCurveEditor::ZoomToFitInternal(EAxisList::Type Axes, const TMap<FCurveMode
 
 		if (Axes & EAxisList::Y)
 		{
-			// Store the min max for each view
-			for (auto ViewIt = GetPanel()->FindViews(CurveID); ViewIt; ++ViewIt)
+			TSharedPtr<SCurveEditorPanel> Panel = WeakPanel.Pin();
+			TSharedPtr<SCurveEditorView> View = WeakView.Pin();
+			if (Panel.IsValid())
 			{
-				TTuple<double, double>* ViewBounds = ViewToOutputBounds.Find(ViewIt.Value());
+				// Store the min max for each view
+				for (auto ViewIt = Panel->FindViews(CurveID); ViewIt; ++ViewIt)
+				{
+					TTuple<double, double>* ViewBounds = ViewToOutputBounds.Find(ViewIt.Value());
+					if (ViewBounds)
+					{
+						ViewBounds->Get<0>() = FMath::Min(ViewBounds->Get<0>(), OutputMin);
+						ViewBounds->Get<1>() = FMath::Max(ViewBounds->Get<1>(), OutputMax);
+					}
+					else
+					{
+						ViewToOutputBounds.Add(ViewIt.Value(), MakeTuple(OutputMin, OutputMax));
+					}
+				}
+			}
+			else if(View.IsValid())
+			{
+				TTuple<double, double>* ViewBounds = ViewToOutputBounds.Find(View.ToSharedRef());
 				if (ViewBounds)
 				{
 					ViewBounds->Get<0>() = FMath::Min(ViewBounds->Get<0>(), OutputMin);
@@ -445,7 +537,7 @@ void FCurveEditor::ZoomToFitInternal(EAxisList::Type Axes, const TMap<FCurveMode
 				}
 				else
 				{
-					ViewToOutputBounds.Add(ViewIt.Value(), MakeTuple(OutputMin, OutputMax));
+					ViewToOutputBounds.Add(View.ToSharedRef(), MakeTuple(OutputMin, OutputMax));
 				}
 			}
 		}
@@ -1487,7 +1579,6 @@ void FCurveEditor::PostUndo(bool bSuccess)
 			Selection.Remove(Set.Key);
 			continue;
 		}
-
 		// Get all of the key handles from this curve.
 		TArray<FKeyHandle> KeyHandles;
 		CurveModel->GetKeys(*this, TNumericLimits<double>::Lowest(), TNumericLimits<double>::Max(), TNumericLimits<double>::Lowest(), TNumericLimits<double>::Max(), KeyHandles);
diff --git a/Engine/Source/Editor/CurveEditor/Private/RichCurveEditorModel.cpp b/Engine/Source/Editor/CurveEditor/Private/RichCurveEditorModel.cpp
--- a/Engine/Source/Editor/CurveEditor/Private/RichCurveEditorModel.cpp
+++ b/Engine/Source/Editor/CurveEditor/Private/RichCurveEditorModel.cpp
@@ -14,9 +14,8 @@
 #include "UObject/Package.h"
 
 
-void RefineCurvePoints(const FRichCurve* RichCurve, double TimeThreshold, float ValueThreshold, TArray<TTuple<double, double>>& InOutPoints)
+void RefineCurvePoints(const FRichCurve& RichCurve, double TimeThreshold, float ValueThreshold, TArray<TTuple<double, double>>& InOutPoints)
 {
-	check(RichCurve);
 	const float InterpTimes[] = { 0.25f, 0.5f, 0.6f };
 
 	for (int32 Index = 0; Index < InOutPoints.Num() - 1; ++Index)
@@ -36,7 +35,7 @@ void RefineCurvePoints(const FRichCurve* RichCurve, double TimeThreshold, float
 
 				EvalTime = FMath::Lerp(Lower.Get<0>(), Upper.Get<0>(), InterpTimes[InterpIndex]);
 
-				float Value = RichCurve->Eval(EvalTime);
+				float Value = RichCurve.Eval(EvalTime);
 
 				const float LinearValue = FMath::Lerp(Lower.Get<1>(), Upper.Get<1>(), InterpTimes[InterpIndex]);
 				if (bSegmentIsLinear)
@@ -63,10 +62,10 @@ void RefineCurvePoints(const FRichCurve* RichCurve, double TimeThreshold, float
 class FRichBufferedCurveModel : public IBufferedCurveModel
 {
 public:
-	FRichBufferedCurveModel(FRichCurve* InRichCurve, TArray<FKeyPosition>&& InKeyPositions, TArray<FKeyAttributes>&& InKeyAttributes,
+	FRichBufferedCurveModel(const FRichCurve& InRichCurve, TArray<FKeyPosition>&& InKeyPositions, TArray<FKeyAttributes>&& InKeyAttributes,
 		const FString& InIntentionName, const double InValueMin, const double InValueMax)
 		: IBufferedCurveModel(MoveTemp(InKeyPositions), MoveTemp(InKeyAttributes), InIntentionName, InValueMin, InValueMax)
-		, RichCurve(*InRichCurve)
+		, RichCurve(InRichCurve)
 	{}
 
 	virtual void DrawCurve(const FCurveEditor& CurveEditor, const FCurveEditorScreenSpace& ScreenSpace, TArray<TTuple<double, double>>& InterpolatingPoints) const override
@@ -92,7 +91,7 @@ public:
 		do
 		{
 			OldSize = InterpolatingPoints.Num();
-			RefineCurvePoints(&RichCurve, TimeThreshold, ValueThreshold, InterpolatingPoints);
+			RefineCurvePoints(RichCurve, TimeThreshold, ValueThreshold, InterpolatingPoints);
 		} while (OldSize != InterpolatingPoints.Num());
 	}
 
@@ -100,23 +99,29 @@ private:
 	FRichCurve RichCurve;
 };
 
-FRichCurveEditorModel::FRichCurveEditorModel(FRichCurve* InRichCurve, UObject* InOwner)
-	: RichCurve(InRichCurve), WeakOwner(InOwner)
+FRichCurveEditorModel::FRichCurveEditorModel( UObject* InOwner)
+	: WeakOwner(InOwner), ClampInputRange(TRange<double>(TNumericLimits<double>::Lowest(), TNumericLimits<double>::Max()))
 {
-	checkf(RichCurve, TEXT("If is not valid to provide a null rich curve to this class"));
 }
 
 const void* FRichCurveEditorModel::GetCurve() const
 {
-	return RichCurve;
+	if(IsValid())
+	{
+		return &GetReadOnlyRichCurve();
+	}
+	return nullptr;
 }
 
 void FRichCurveEditorModel::Modify()
 {
 	if (UObject* Owner = WeakOwner.Get())
 	{
-		Owner->SetFlags(RF_Transactional);
-		Owner->Modify();
+		if(IsValid())
+		{
+			Owner->SetFlags(RF_Transactional);
+			Owner->Modify();
+		}
 	}
 }
 
@@ -124,32 +129,38 @@ void FRichCurveEditorModel::AddKeys(TArrayView<const FKeyPosition> InKeyPosition
 {
 	check(InKeyPositions.Num() == InKeyAttributes.Num() && (!OutKeyHandles || OutKeyHandles->Num() == InKeyPositions.Num()));
 
-	UObject* Owner = WeakOwner.Get();
-	if (Owner)
+	if (UObject* Owner = WeakOwner.Get())
 	{
-		Owner->Modify();
-
-		TArray<FKeyHandle> NewKeyHandles;
-		NewKeyHandles.SetNumUninitialized(InKeyPositions.Num());
-
-		for (int32 Index = 0; Index < InKeyPositions.Num(); ++Index)
+		if(IsValid())
 		{
-			FKeyPosition   Position   = InKeyPositions[Index];
-			FKeyAttributes Attributes = InKeyAttributes[Index];
+			Owner->Modify();
+
+			TArray<FKeyHandle> NewKeyHandles;
+			NewKeyHandles.SetNumUninitialized(InKeyPositions.Num());
 
-			FKeyHandle     NewHandle = RichCurve->AddKey(Position.InputValue, Position.OutputValue);
-			FRichCurveKey* NewKey    = &RichCurve->GetKey(NewHandle);
+			FRichCurve& RichCurve = GetRichCurve();
 
-			NewKeyHandles[Index] = NewHandle;
-			if (OutKeyHandles)
+			for (int32 Index = 0; Index < InKeyPositions.Num(); ++Index)
 			{
-				(*OutKeyHandles)[Index] = NewHandle;
+				FKeyPosition   Position   = InKeyPositions[Index];
+				FKeyAttributes Attributes = InKeyAttributes[Index];
+
+				const TRange<double> InputRange = ClampInputRange.Get();
+				FKeyHandle     NewHandle = RichCurve.AddKey(FMath::Clamp(Position.InputValue, InputRange.GetLowerBoundValue(), InputRange.GetUpperBoundValue()), Position.OutputValue);
+
+				NewKeyHandles[Index] = NewHandle;
+				if (OutKeyHandles)
+				{
+					(*OutKeyHandles)[Index] = NewHandle;
+				}
 			}
-		}
 
-		// We reuse SetKeyAttributes here as there is complex logic determining which parts of the attributes are valid to pass on.
-		// For now we need to duplicate the new key handle array due to API mismatch. This will auto-calculate tangents if required.
-		SetKeyAttributes(NewKeyHandles, InKeyAttributes, EPropertyChangeType::ValueSet);
+			// We reuse SetKeyAttributes here as there is complex logic determining which parts of the attributes are valid to pass on.
+			// For now we need to duplicate the new key handle array due to API mismatch. This will auto-calculate tangents if required.
+			SetKeyAttributes(NewKeyHandles, InKeyAttributes);
+
+			CurveModifiedDelegate.Broadcast();
+		}
 	}
 }
 
@@ -157,8 +168,11 @@ bool FRichCurveEditorModel::Evaluate(double Time, double& OutValue) const
 {
 	if (UObject* Owner = WeakOwner.Get())
 	{
-		OutValue = RichCurve->Eval(Time);
-		return true;
+		if(IsValid())
+		{
+			OutValue = GetReadOnlyRichCurve().Eval(Time);
+			return true;
+		}
 	}
 
 	return false;
@@ -168,10 +182,16 @@ void FRichCurveEditorModel::RemoveKeys(TArrayView<const FKeyHandle> InKeys)
 {
 	if (UObject* Owner = WeakOwner.Get())
 	{
-		Owner->Modify();
-		for (FKeyHandle Handle : InKeys)
+		if(IsValid())
 		{
-			RichCurve->DeleteKey(Handle);
+			Owner->Modify();
+			FRichCurve& RichCurve = GetRichCurve();
+			for (FKeyHandle Handle : InKeys)
+			{
+				RichCurve.DeleteKey(Handle);
+			}
+
+			CurveModifiedDelegate.Broadcast();
 		}
 	}
 }
@@ -180,30 +200,35 @@ void FRichCurveEditorModel::DrawCurve(const FCurveEditor& CurveEditor, const FCu
 {
 	if (UObject* Owner = WeakOwner.Get())
 	{
-		const double StartTimeSeconds = ScreenSpace.GetInputMin();
-		const double EndTimeSeconds   = ScreenSpace.GetInputMax();
-		const double TimeThreshold    = FMath::Max(0.0001, 1.0 / ScreenSpace.PixelsPerInput());
-		const double ValueThreshold   = FMath::Max(0.0001, 1.0 / ScreenSpace.PixelsPerOutput());
+		if(IsValid())
+		{
+			const double StartTimeSeconds = ScreenSpace.GetInputMin();
+			const double EndTimeSeconds   = ScreenSpace.GetInputMax();
+			const double TimeThreshold    = FMath::Max(0.0001, 1.0 / ScreenSpace.PixelsPerInput());
+			const double ValueThreshold   = FMath::Max(0.0001, 1.0 / ScreenSpace.PixelsPerOutput());
 
-		InOutPoints.Add(MakeTuple(StartTimeSeconds, double(RichCurve->Eval(StartTimeSeconds))));
+			const FRichCurve& RichCurve = GetReadOnlyRichCurve();
 
-		for (const FRichCurveKey& Key : RichCurve->GetConstRefOfKeys())
-		{
-			if (Key.Time > StartTimeSeconds && Key.Time < EndTimeSeconds)
+			InOutPoints.Add(MakeTuple(StartTimeSeconds, double(RichCurve.Eval(StartTimeSeconds))));
+
+			for (const FRichCurveKey& Key : RichCurve.GetConstRefOfKeys())
 			{
-				InOutPoints.Add(MakeTuple(double(Key.Time), double(Key.Value)));
+				if (Key.Time > StartTimeSeconds && Key.Time < EndTimeSeconds)
+				{
+					InOutPoints.Add(MakeTuple(double(Key.Time), double(Key.Value)));
+				}
 			}
-		}
 
-		InOutPoints.Add(MakeTuple(EndTimeSeconds, double(RichCurve->Eval(EndTimeSeconds))));
+			InOutPoints.Add(MakeTuple(EndTimeSeconds, double(RichCurve.Eval(EndTimeSeconds))));
 
-		int32 OldSize = InOutPoints.Num();
-		do
-		{
-			OldSize = InOutPoints.Num();
-			RefineCurvePoints(RichCurve, TimeThreshold, ValueThreshold, InOutPoints);
+			int32 OldSize = InOutPoints.Num();
+			do
+			{
+				OldSize = InOutPoints.Num();
+				RefineCurvePoints(RichCurve, TimeThreshold, ValueThreshold, InOutPoints);
+			}
+			while(OldSize != InOutPoints.Num());
 		}
-		while(OldSize != InOutPoints.Num());
 	}
 }
 
@@ -211,12 +236,16 @@ void FRichCurveEditorModel::GetKeys(const FCurveEditor& CurveEditor, double MinT
 {
 	if (UObject* Owner = WeakOwner.Get())
 	{
-		for (auto It = RichCurve->GetKeyHandleIterator(); It; ++It)
+		if(IsValid())
 		{
-			const FRichCurveKey& Key = RichCurve->GetKey(*It);
-			if (Key.Time >= MinTime && Key.Time <= MaxTime && Key.Value >= MinValue && Key.Value <= MaxValue)
+			const FRichCurve& RichCurve = GetReadOnlyRichCurve();
+			for (auto It = RichCurve.GetKeyHandleIterator(); It; ++It)
 			{
-				OutKeyHandles.Add(*It);
+				const FRichCurveKey& Key = RichCurve.GetKeyRef(*It);
+				if (Key.Time >= MinTime && Key.Time <= MaxTime && Key.Value >= MinValue && Key.Value <= MaxValue)
+				{
+					OutKeyHandles.Add(*It);
+				}
 			}
 		}
 	}
@@ -234,7 +263,7 @@ void FRichCurveEditorModel::GetKeyDrawInfo(ECurvePointType PointType, const FKey
 		// All keys are the same size by default
 		OutDrawInfo.ScreenSize = FVector2D(11, 11);
 
-		ERichCurveInterpMode KeyType = RichCurve->IsKeyHandleValid(InKeyHandle) ? RichCurve->GetKey(InKeyHandle).InterpMode.GetValue() : RCIM_None;
+		ERichCurveInterpMode KeyType = (IsValid() && GetReadOnlyRichCurve().IsKeyHandleValid(InKeyHandle)) ? GetReadOnlyRichCurve().GetKeyRef(InKeyHandle).InterpMode.GetValue() : RCIM_None;
 		switch (KeyType)
 		{
 		case ERichCurveInterpMode::RCIM_Constant:
@@ -261,14 +290,18 @@ void FRichCurveEditorModel::GetKeyPositions(TArrayView<const FKeyHandle> InKeys,
 {
 	if (UObject* Owner = WeakOwner.Get())
 	{
-		for (int32 Index = 0; Index < InKeys.Num(); ++Index)
+		if(IsValid())
 		{
-			if (RichCurve->IsKeyHandleValid(InKeys[Index]))
+			const FRichCurve& RichCurve = GetReadOnlyRichCurve();
+			for (int32 Index = 0; Index < InKeys.Num(); ++Index)
 			{
-				const FRichCurveKey& Key = RichCurve->GetKey(InKeys[Index]);
+				if (RichCurve.IsKeyHandleValid(InKeys[Index]))
+				{
+					const FRichCurveKey& Key = RichCurve.GetKeyRef(InKeys[Index]);
 
-				OutKeyPositions[Index].InputValue  = Key.Time;
-				OutKeyPositions[Index].OutputValue = Key.Value;
+					OutKeyPositions[Index].InputValue  = Key.Time;
+					OutKeyPositions[Index].OutputValue = Key.Value;
+				}
 			}
 		}
 	}
@@ -283,21 +316,28 @@ void FRichCurveEditorModel::SetKeyPositions(TArrayView<const FKeyHandle> InKeys,
 
 	if (UObject* Owner = WeakOwner.Get())
 	{
-		Owner->Modify();
-
-		for (int32 Index = 0; Index < InKeys.Num(); ++Index)
+		if(IsValid())
 		{
-			FKeyHandle Handle = InKeys[Index];
-			if (RichCurve->IsKeyHandleValid(Handle))
+			Owner->Modify();
+
+			FRichCurve& RichCurve = GetRichCurve();
+			for (int32 Index = 0; Index < InKeys.Num(); ++Index)
 			{
-				// Set key time last so we don't have to worry about the key handle changing
-				RichCurve->GetKey(Handle).Value = InKeyPositions[Index].OutputValue;
-				RichCurve->SetKeyTime(Handle, InKeyPositions[Index].InputValue);
+				FKeyHandle Handle = InKeys[Index];
+				if (RichCurve.IsKeyHandleValid(Handle))
+				{
+					// Set key time last so we don't have to worry about the key handle changing
+					RichCurve.GetKey(Handle).Value = InKeyPositions[Index].OutputValue;
+					const TRange<double> InputRange = ClampInputRange.Get();
+					RichCurve.SetKeyTime(Handle, FMath::Clamp(InKeyPositions[Index].InputValue, InputRange.GetLowerBoundValue(), InputRange.GetUpperBoundValue()));
+				}
 			}
+			RichCurve.AutoSetTangents();
+			FPropertyChangedEvent PropertyChangeStruct(nullptr, EPropertyChangeType::ValueSet);
+			Owner->PostEditChangeProperty(PropertyChangeStruct);
+
+			CurveModifiedDelegate.Broadcast();
 		}
-		RichCurve->AutoSetTangents();
-		FPropertyChangedEvent PropertyChangeStruct(nullptr, ChangeType);
-		Owner->PostEditChangeProperty(PropertyChangeStruct);
 	}
 }
 
@@ -305,43 +345,47 @@ void FRichCurveEditorModel::GetKeyAttributes(TArrayView<const FKeyHandle> InKeys
 {
 	if (UObject* Owner = WeakOwner.Get())
 	{
-		const TArray<FRichCurveKey>& AllKeys = RichCurve->GetConstRefOfKeys();
-		if (AllKeys.Num() == 0)
+		if(IsValid())
 		{
-			return;
-		}
+			const FRichCurve& RichCurve = GetReadOnlyRichCurve();
+			const TArray<FRichCurveKey>& AllKeys = RichCurve.GetConstRefOfKeys();
+			if (AllKeys.Num() == 0)
+			{
+				return;
+			}
 
-		const FRichCurveKey* FirstKey = &AllKeys[0];
-		const FRichCurveKey* LastKey  = &AllKeys.Last();
+			const FRichCurveKey* FirstKey = &AllKeys[0];
+			const FRichCurveKey* LastKey  = &AllKeys.Last();
 
-		for (int32 Index = 0; Index < InKeys.Num(); ++Index)
-		{
-			if (RichCurve->IsKeyHandleValid(InKeys[Index]))
+			for (int32 Index = 0; Index < InKeys.Num(); ++Index)
 			{
-				const FRichCurveKey* ThisKey    = &RichCurve->GetKey(InKeys[Index]);
-				FKeyAttributes&      Attributes = OutAttributes[Index];
+				if (RichCurve.IsKeyHandleValid(InKeys[Index]))
+				{
+					const FRichCurveKey& ThisKey    = RichCurve.GetKeyRef(InKeys[Index]);
+					FKeyAttributes&      Attributes = OutAttributes[Index];
 
-				Attributes.SetInterpMode(ThisKey->InterpMode);
+					Attributes.SetInterpMode(ThisKey.InterpMode);
 
-				if (ThisKey->InterpMode != RCIM_Constant && ThisKey->InterpMode != RCIM_Linear)
-				{
-					Attributes.SetTangentMode(ThisKey->TangentMode);
-					if (ThisKey != FirstKey)
+					if (ThisKey.InterpMode != RCIM_Constant && ThisKey.InterpMode != RCIM_Linear)
 					{
-						Attributes.SetArriveTangent(ThisKey->ArriveTangent);
-					}
+						Attributes.SetTangentMode(ThisKey.TangentMode);
+						if (&ThisKey != FirstKey)
+						{
+							Attributes.SetArriveTangent(ThisKey.ArriveTangent);
+						}
 
-					if (ThisKey != LastKey)
-					{
-						Attributes.SetLeaveTangent(ThisKey->LeaveTangent);
-					}
-					if (ThisKey->InterpMode == RCIM_Cubic)
-					{
-						Attributes.SetTangentWeightMode(ThisKey->TangentWeightMode);
-						if (ThisKey->TangentWeightMode != RCTWM_WeightedNone)
+						if (&ThisKey != LastKey)
 						{
-							Attributes.SetArriveTangentWeight(ThisKey->ArriveTangentWeight);
-							Attributes.SetLeaveTangentWeight(ThisKey->LeaveTangentWeight);
+							Attributes.SetLeaveTangent(ThisKey.LeaveTangent);
+						}
+						if (ThisKey.InterpMode == RCIM_Cubic)
+						{
+							Attributes.SetTangentWeightMode(ThisKey.TangentWeightMode);
+							if (ThisKey.TangentWeightMode != RCTWM_WeightedNone)
+							{
+								Attributes.SetArriveTangentWeight(ThisKey.ArriveTangentWeight);
+								Attributes.SetLeaveTangentWeight(ThisKey.LeaveTangentWeight);
+							}
 						}
 					}
 				}
@@ -359,137 +403,142 @@ void FRichCurveEditorModel::SetKeyAttributes(TArrayView<const FKeyHandle> InKeys
 
 	if (UObject* Owner = WeakOwner.Get())
 	{
-		const TArray<FRichCurveKey>& AllKeys = RichCurve->GetConstRefOfKeys();
-		if (AllKeys.Num() == 0)
+		if(IsValid())
 		{
-			return;
-		}
+			FRichCurve& RichCurve = GetRichCurve();
+			const TArray<FRichCurveKey>& AllKeys = RichCurve.GetConstRefOfKeys();
+			if (AllKeys.Num() == 0)
+			{
+				return;
+			}
 
-		Owner->Modify();
+			Owner->Modify();
 
-		const FRichCurveKey* FirstKey = &AllKeys[0];
-		const FRichCurveKey* LastKey  = &AllKeys.Last();
+			const FRichCurveKey* FirstKey = &AllKeys[0];
+			const FRichCurveKey* LastKey  = &AllKeys.Last();
 
-		bool bAutoSetTangents = false;
+			bool bAutoSetTangents = false;
 
-		for (int32 Index = 0; Index < InKeys.Num(); ++Index)
-		{
-			if (RichCurve->IsKeyHandleValid(InKeys[Index]))
+			for (int32 Index = 0; Index < InKeys.Num(); ++Index)
 			{
-				FRichCurveKey*        ThisKey    = &RichCurve->GetKey(InKeys[Index]);
-				const FKeyAttributes& Attributes = InAttributes[Index];
-
-				if (Attributes.HasInterpMode())    { ThisKey->InterpMode  = Attributes.GetInterpMode();  bAutoSetTangents = true; }
-				if (Attributes.HasTangentMode())
+				if (RichCurve.IsKeyHandleValid(InKeys[Index]))
 				{
-					ThisKey->TangentMode = Attributes.GetTangentMode();
-					if (ThisKey->TangentMode == RCTM_Auto)
+					FRichCurveKey*        ThisKey    = &RichCurve.GetKey(InKeys[Index]);
+					const FKeyAttributes& Attributes = InAttributes[Index];
+
+					if (Attributes.HasInterpMode())    { ThisKey->InterpMode  = Attributes.GetInterpMode();  bAutoSetTangents = true; }
+					if (Attributes.HasTangentMode())
 					{
-						ThisKey->TangentWeightMode = RCTWM_WeightedNone;
+						ThisKey->TangentMode = Attributes.GetTangentMode();
+						if (ThisKey->TangentMode == RCTM_Auto)
+						{
+							ThisKey->TangentWeightMode = RCTWM_WeightedNone;
+						}
+						bAutoSetTangents = true;
 					}
-					bAutoSetTangents = true;
-				}
-				if (Attributes.HasTangentWeightMode()) 
-				{ 
-					if (ThisKey->TangentWeightMode == RCTWM_WeightedNone) //set tangent weights to default use
-					{
-						const float OneThird = 1.0f / 3.0f;
-
-						//calculate a tangent weight based upon tangent and time difference
-						//calculate arrive tangent weight
-						if (ThisKey != FirstKey)
+					if (Attributes.HasTangentWeightMode()) 
+					{ 
+						if (ThisKey->TangentWeightMode == RCTWM_WeightedNone) //set tangent weights to default use
 						{
-							const float Y = ThisKey->ArriveTangent;
-							ThisKey->ArriveTangentWeight = FMath::Sqrt(1.f + Y*Y) * OneThird;
+							const float OneThird = 1.0f / 3.0f;
+
+							//calculate a tangent weight based upon tangent and time difference
+							//calculate arrive tangent weight
+							if (ThisKey != FirstKey)
+							{
+								const float Y = ThisKey->ArriveTangent;
+								ThisKey->ArriveTangentWeight = FMath::Sqrt(1.f + Y*Y) * OneThird;
+							}
+							//calculate leave weight
+							if(ThisKey != LastKey)
+							{
+								const float Y = ThisKey->LeaveTangent;
+								ThisKey->LeaveTangentWeight = FMath::Sqrt(1.f + Y*Y) * OneThird;
+							}
 						}
-						//calculate leave weight
-						if(ThisKey != LastKey)
+						ThisKey->TangentWeightMode = Attributes.GetTangentWeightMode();
+
+						if( ThisKey->TangentWeightMode != RCTWM_WeightedNone )
 						{
-							const float Y = ThisKey->LeaveTangent;
-							ThisKey->LeaveTangentWeight = FMath::Sqrt(1.f + Y*Y) * OneThird;
+							if (ThisKey->TangentMode != RCTM_User && ThisKey->TangentMode != RCTM_Break)
+							{
+								ThisKey->TangentMode = RCTM_User;
+							}
 						}
 					}
-					ThisKey->TangentWeightMode = Attributes.GetTangentWeightMode();
 
-					if( ThisKey->TangentWeightMode != RCTWM_WeightedNone )
+					if (Attributes.HasArriveTangent())
 					{
-						if (ThisKey->TangentMode != RCTM_User && ThisKey->TangentMode != RCTM_Break)
+						if (ThisKey->TangentMode == RCTM_Auto)
 						{
 							ThisKey->TangentMode = RCTM_User;
+							ThisKey->TangentWeightMode = RCTWM_WeightedNone;
 						}
-					}
-				}
 
-				if (Attributes.HasArriveTangent())
-				{
-					if (ThisKey->TangentMode == RCTM_Auto)
-					{
-						ThisKey->TangentMode = RCTM_User;
-						ThisKey->TangentWeightMode = RCTWM_WeightedNone;
+						ThisKey->ArriveTangent = Attributes.GetArriveTangent();
+						if (ThisKey->InterpMode == RCIM_Cubic && ThisKey->TangentMode != RCTM_Break)
+						{
+							ThisKey->LeaveTangent = ThisKey->ArriveTangent;
+						}
 					}
 
-					ThisKey->ArriveTangent = Attributes.GetArriveTangent();
-					if (ThisKey->InterpMode == RCIM_Cubic && ThisKey->TangentMode != RCTM_Break)
+					if (Attributes.HasLeaveTangent())
 					{
-						ThisKey->LeaveTangent = ThisKey->ArriveTangent;
-					}
-				}
+						if (ThisKey->TangentMode == RCTM_Auto)
+						{
+							ThisKey->TangentMode = RCTM_User;
+							ThisKey->TangentWeightMode = RCTWM_WeightedNone;
+						}
 
-				if (Attributes.HasLeaveTangent())
-				{
-					if (ThisKey->TangentMode == RCTM_Auto)
-					{
-						ThisKey->TangentMode = RCTM_User;
-						ThisKey->TangentWeightMode = RCTWM_WeightedNone;
+						ThisKey->LeaveTangent = Attributes.GetLeaveTangent();
+						if (ThisKey->InterpMode == RCIM_Cubic && ThisKey->TangentMode != RCTM_Break)
+						{
+							ThisKey->ArriveTangent = ThisKey->LeaveTangent;
+						}
 					}
 
-					ThisKey->LeaveTangent = Attributes.GetLeaveTangent();
-					if (ThisKey->InterpMode == RCIM_Cubic && ThisKey->TangentMode != RCTM_Break)
+					if (Attributes.HasArriveTangentWeight())
 					{
-						ThisKey->ArriveTangent = ThisKey->LeaveTangent;
-					}
-				}
+						if (ThisKey->TangentMode == RCTM_Auto)
+						{
+							ThisKey->TangentMode = RCTM_User;
+							ThisKey->TangentWeightMode = RCTWM_WeightedNone;
+						}
 
-				if (Attributes.HasArriveTangentWeight())
-				{
-					if (ThisKey->TangentMode == RCTM_Auto)
-					{
-						ThisKey->TangentMode = RCTM_User;
-						ThisKey->TangentWeightMode = RCTWM_WeightedNone;
+						ThisKey->ArriveTangentWeight = Attributes.GetArriveTangentWeight();
+						if (ThisKey->InterpMode == RCIM_Cubic && ThisKey->TangentMode != RCTM_Break)
+						{
+							ThisKey->LeaveTangentWeight = ThisKey->ArriveTangentWeight;
+						}
 					}
 
-					ThisKey->ArriveTangentWeight = Attributes.GetArriveTangentWeight();
-					if (ThisKey->InterpMode == RCIM_Cubic && ThisKey->TangentMode != RCTM_Break)
+					if (Attributes.HasLeaveTangentWeight())
 					{
-						ThisKey->LeaveTangentWeight = ThisKey->ArriveTangentWeight;
-					}
-				}
-
-				if (Attributes.HasLeaveTangentWeight())
-				{
 				
-					if (ThisKey->TangentMode == RCTM_Auto)
-					{
-						ThisKey->TangentMode = RCTM_User;
-						ThisKey->TangentWeightMode = RCTWM_WeightedNone;
-					}
+						if (ThisKey->TangentMode == RCTM_Auto)
+						{
+							ThisKey->TangentMode = RCTM_User;
+							ThisKey->TangentWeightMode = RCTWM_WeightedNone;
+						}
 
-					ThisKey->LeaveTangentWeight = Attributes.GetLeaveTangentWeight();
-					if (ThisKey->InterpMode == RCIM_Cubic && ThisKey->TangentMode != RCTM_Break)
-					{
-						ThisKey->ArriveTangentWeight = ThisKey->LeaveTangentWeight;
+						ThisKey->LeaveTangentWeight = Attributes.GetLeaveTangentWeight();
+						if (ThisKey->InterpMode == RCIM_Cubic && ThisKey->TangentMode != RCTM_Break)
+						{
+							ThisKey->ArriveTangentWeight = ThisKey->LeaveTangentWeight;
+						}
 					}
 				}
 			}
-		}
 
-		if (bAutoSetTangents)
-		{
-			RichCurve->AutoSetTangents();
-		}
+			if (bAutoSetTangents)
+			{
+				RichCurve.AutoSetTangents();
+			}
 
-		FPropertyChangedEvent PropertyChangeStruct(nullptr, ChangeType);
-		Owner->PostEditChangeProperty(PropertyChangeStruct);
+			FPropertyChangedEvent PropertyChangeStruct(nullptr, EPropertyChangeType::ValueSet);
+			Owner->PostEditChangeProperty(PropertyChangeStruct);
+			CurveModifiedDelegate.Broadcast();
+		}
 	}
 }
 
@@ -497,8 +546,12 @@ void FRichCurveEditorModel::GetCurveAttributes(FCurveAttributes& OutCurveAttribu
 {
 	if (UObject* Owner = WeakOwner.Get())
 	{
-		OutCurveAttributes.SetPreExtrapolation(RichCurve->PreInfinityExtrap);
-		OutCurveAttributes.SetPostExtrapolation(RichCurve->PostInfinityExtrap);
+		if(IsValid())
+		{
+			const FRichCurve& RichCurve = GetReadOnlyRichCurve();
+			OutCurveAttributes.SetPreExtrapolation(RichCurve.PreInfinityExtrap);
+			OutCurveAttributes.SetPostExtrapolation(RichCurve.PostInfinityExtrap);
+		}
 	}
 }
 
@@ -506,31 +559,43 @@ void FRichCurveEditorModel::SetCurveAttributes(const FCurveAttributes& InCurveAt
 {
 	if (UObject* Owner = WeakOwner.Get())
 	{
-		Owner->Modify();
-
-		if (InCurveAttributes.HasPreExtrapolation())
+		if(IsValid())
 		{
-			RichCurve->PreInfinityExtrap = InCurveAttributes.GetPreExtrapolation();
-		}
+			Owner->Modify();
 
-		if (InCurveAttributes.HasPostExtrapolation())
-		{
-			RichCurve->PostInfinityExtrap = InCurveAttributes.GetPostExtrapolation();
-		}
+			FRichCurve& RichCurve = GetRichCurve();
+
+			if (InCurveAttributes.HasPreExtrapolation())
+			{
+				RichCurve.PreInfinityExtrap = InCurveAttributes.GetPreExtrapolation();
+			}
 
-		FPropertyChangedEvent PropertyChangeStruct(nullptr, EPropertyChangeType::ValueSet);
-		Owner->PostEditChangeProperty(PropertyChangeStruct);
+			if (InCurveAttributes.HasPostExtrapolation())
+			{
+				RichCurve.PostInfinityExtrap = InCurveAttributes.GetPostExtrapolation();
+			}
+
+			FPropertyChangedEvent PropertyChangeStruct(nullptr, EPropertyChangeType::ValueSet);
+			Owner->PostEditChangeProperty(PropertyChangeStruct);
+			CurveModifiedDelegate.Broadcast();
+		}
 	}
 }
 
 void FRichCurveEditorModel::CreateKeyProxies(TArrayView<const FKeyHandle> InKeyHandles, TArrayView<UObject*> OutObjects)
 {
-	for (int32 Index = 0; Index < InKeyHandles.Num(); ++Index)
+	if (UObject* Owner = WeakOwner.Get())
 	{
-		URichCurveKeyProxy* NewProxy = NewObject<URichCurveKeyProxy>(GetTransientPackage(), NAME_None);
+		if(IsValid())
+		{
+			for (int32 Index = 0; Index < InKeyHandles.Num(); ++Index)
+			{
+				URichCurveKeyProxy* NewProxy = NewObject<URichCurveKeyProxy>(GetTransientPackage(), NAME_None);
 
-		NewProxy->Initialize(InKeyHandles[Index], RichCurve, WeakOwner);
-		OutObjects[Index] = NewProxy;
+				NewProxy->Initialize(InKeyHandles[Index], this, WeakOwner);
+				OutObjects[Index] = NewProxy;
+			}
+		}
 	}
 }
 
@@ -538,23 +603,28 @@ TUniquePtr<IBufferedCurveModel> FRichCurveEditorModel::CreateBufferedCurveCopy()
 {
 	if (UObject* Owner = WeakOwner.Get())
 	{
-		TArray<FKeyHandle> TargetKeyHandles;
-		for (auto It = RichCurve->GetKeyHandleIterator(); It; ++It)
+		if(IsValid())
 		{
-			TargetKeyHandles.Add(*It);
-		}
+			const FRichCurve& RichCurve = GetReadOnlyRichCurve();
 
-		TArray<FKeyPosition> KeyPositions;
-		KeyPositions.SetNumUninitialized(TargetKeyHandles.Num());
-		TArray<FKeyAttributes> KeyAttributes;
-		KeyAttributes.SetNumUninitialized(TargetKeyHandles.Num());
-		GetKeyPositions(TargetKeyHandles, KeyPositions);
-		GetKeyAttributes(TargetKeyHandles, KeyAttributes);
+			TArray<FKeyHandle> TargetKeyHandles;
+			for (auto It = RichCurve.GetKeyHandleIterator(); It; ++It)
+			{
+				TargetKeyHandles.Add(*It);
+			}
+
+			TArray<FKeyPosition> KeyPositions;
+			KeyPositions.SetNumUninitialized(TargetKeyHandles.Num());
+			TArray<FKeyAttributes> KeyAttributes;
+			KeyAttributes.SetNumUninitialized(TargetKeyHandles.Num());
+			GetKeyPositions(TargetKeyHandles, KeyPositions);
+			GetKeyAttributes(TargetKeyHandles, KeyAttributes);
 
-		double ValueMin = 0.f, ValueMax = 1.f;
-		GetValueRange(ValueMin, ValueMax);
+			double ValueMin = 0.f, ValueMax = 1.f;
+			GetValueRange(ValueMin, ValueMax);
 
-		return MakeUnique<FRichBufferedCurveModel>(RichCurve, MoveTemp(KeyPositions), MoveTemp(KeyAttributes), GetIntentionName(), ValueMin, ValueMax);
+			return MakeUnique<FRichBufferedCurveModel>(RichCurve, MoveTemp(KeyPositions), MoveTemp(KeyAttributes), GetIntentionName(), ValueMin, ValueMax);
+		}
 	}
 
 	return nullptr;
@@ -564,11 +634,15 @@ void FRichCurveEditorModel::GetTimeRange(double& MinTime, double& MaxTime) const
 {
 	if (UObject* Owner = WeakOwner.Get())
 	{
-		float MinTimeFloat = 0.f, MaxTimeFloat = 0.f;
-		RichCurve->GetTimeRange(MinTimeFloat, MaxTimeFloat);
+		if(IsValid())
+		{
+			float MinTimeFloat = 0.f, MaxTimeFloat = 0.f;
+			const FRichCurve& RichCurve = GetReadOnlyRichCurve();
+			RichCurve.GetTimeRange(MinTimeFloat, MaxTimeFloat);
 
-		MinTime = MinTimeFloat;
-		MaxTime = MaxTimeFloat;
+			MinTime = MinTimeFloat;
+			MaxTime = MaxTimeFloat;
+		}
 	}
 }
 
@@ -576,38 +650,70 @@ void FRichCurveEditorModel::GetValueRange(double& MinValue, double& MaxValue) co
 {
 	if (UObject* Owner = WeakOwner.Get())
 	{
-		float MinValueFloat = 0.f, MaxValueFloat = 0.f;
-		RichCurve->GetValueRange(MinValueFloat, MaxValueFloat);
+		if(IsValid())
+		{
+			float MinValueFloat = 0.f, MaxValueFloat = 0.f;
+			const FRichCurve& RichCurve = GetReadOnlyRichCurve();
+			RichCurve.GetValueRange(MinValueFloat, MaxValueFloat);
 
-		MinValue = MinValueFloat;
-		MaxValue = MaxValueFloat;
+			MinValue = MinValueFloat;
+			MaxValue = MaxValueFloat;
+		}
 	}
 }
 
 int32 FRichCurveEditorModel::GetNumKeys() const
 {
-	return RichCurve->GetNumKeys();
+	if (UObject* Owner = WeakOwner.Get())
+	{
+		if(IsValid())
+		{
+			return GetReadOnlyRichCurve().GetNumKeys();
+		}
+	}
+	return 0;
 }
 
 void FRichCurveEditorModel::GetNeighboringKeys(const FKeyHandle InKeyHandle, TOptional<FKeyHandle>& OutPreviousKeyHandle, TOptional<FKeyHandle>& OutNextKeyHandle) const
 {
 	if (UObject* Owner = WeakOwner.Get())
 	{
-		if (RichCurve->IsKeyHandleValid(InKeyHandle))
+		if(IsValid())
 		{
-			FKeyHandle NextKeyHandle = RichCurve->GetNextKey(InKeyHandle);
-
-			if (RichCurve->IsKeyHandleValid(NextKeyHandle))
+			const FRichCurve& RichCurve = GetReadOnlyRichCurve();
+			if (RichCurve.IsKeyHandleValid(InKeyHandle))
 			{
-				OutNextKeyHandle = NextKeyHandle;
-			}
+				FKeyHandle NextKeyHandle = RichCurve.GetNextKey(InKeyHandle);
 
-			FKeyHandle PreviousKeyHandle = RichCurve->GetPreviousKey(InKeyHandle);
+				if (RichCurve.IsKeyHandleValid(NextKeyHandle))
+				{
+					OutNextKeyHandle = NextKeyHandle;
+				}
 
-			if (RichCurve->IsKeyHandleValid(PreviousKeyHandle))
-			{
-				OutPreviousKeyHandle = PreviousKeyHandle;
+				FKeyHandle PreviousKeyHandle = RichCurve.GetPreviousKey(InKeyHandle);
+
+				if (RichCurve.IsKeyHandleValid(PreviousKeyHandle))
+				{
+					OutPreviousKeyHandle = PreviousKeyHandle;
+				}
 			}
 		}
 	}
 }
+
+FRichCurveEditorModelRaw::FRichCurveEditorModelRaw(FRichCurve* InRichCurve, UObject* InOwner)
+	: FRichCurveEditorModel(InOwner)
+	, RichCurve(InRichCurve)
+{
+	checkf(RichCurve, TEXT("If is not valid to provide a null rich curve to this class"));
+}
+
+FRichCurve& FRichCurveEditorModelRaw::GetRichCurve()
+{
+	return *RichCurve;
+}
+
+const FRichCurve& FRichCurveEditorModelRaw::GetReadOnlyRichCurve() const
+{
+	return *RichCurve;
+}
\ No newline at end of file
diff --git a/Engine/Source/Editor/CurveEditor/Private/RichCurveKeyProxy.h b/Engine/Source/Editor/CurveEditor/Private/RichCurveKeyProxy.h
--- a/Engine/Source/Editor/CurveEditor/Private/RichCurveKeyProxy.h
+++ b/Engine/Source/Editor/CurveEditor/Private/RichCurveKeyProxy.h
@@ -7,6 +7,7 @@
 #include "UObject/WeakObjectPtrTemplates.h"
 #include "CurveEditorKeyProxy.h"
 #include "Curves/RichCurve.h"
+#include "RichCurveEditorModel.h"
 
 #include "RichCurveKeyProxy.generated.h"
 
@@ -19,16 +20,22 @@ public:
 	/**
 	 * Initialize this key proxy object by caching the underlying key object, and retrieving the time/value each tick
 	 */
-	void Initialize(FKeyHandle InKeyHandle, FRichCurve* InRichCurve, TWeakObjectPtr<UObject> InWeakOwner)
+	void Initialize(FKeyHandle InKeyHandle, FRichCurveEditorModel* InModel, TWeakObjectPtr<UObject> InWeakOwner)
 	{
 		KeyHandle  = InKeyHandle;
-		RichCurve  = InRichCurve;
 		WeakOwner  = InWeakOwner;
+		Model      = InModel;
 
-		UObject* Owner = WeakOwner.Get();
-		if (Owner && RichCurve->IsKeyHandleValid(KeyHandle))
+		if(UObject* Owner = WeakOwner.Get())
 		{
-			Value = RichCurve->GetKey(KeyHandle);
+			if(Model->IsValid())
+			{
+				const FRichCurve& RichCurve = Model->GetReadOnlyRichCurve();
+				if (RichCurve.IsKeyHandleValid(KeyHandle))
+				{
+					Value = RichCurve.GetKey(KeyHandle);
+				}
+			}
 		}
 	}
 
@@ -38,32 +45,46 @@ private:
 	{
 		Super::PostEditChangeProperty(PropertyChangedEvent);
 
-		UObject* Owner = WeakOwner.Get();
-		if (Owner && RichCurve->IsKeyHandleValid(KeyHandle))
+		if(UObject* Owner = WeakOwner.Get())
 		{
-			Owner->Modify();
-
-			FRichCurveKey& ActualKey = RichCurve->GetKey(KeyHandle);
-
-			const float PreviousTime = ActualKey.Time;
-			const float NewTime      = Value.Time;
-
-			// Copy the key properties
-			ActualKey = Value;
-			ActualKey.Time = PreviousTime;
-			if (PreviousTime != NewTime)
+			if(Model->IsValid())
 			{
-				RichCurve->SetKeyTime(KeyHandle, NewTime);
+				FRichCurve& RichCurve = Model->GetRichCurve();
+				if(RichCurve.IsKeyHandleValid(KeyHandle))
+				{
+					Owner->Modify();
+
+					FRichCurveKey& ActualKey = RichCurve.GetKey(KeyHandle);
+
+					const float PreviousTime = ActualKey.Time;
+					const float NewTime      = Value.Time;
+
+					// Copy the key properties
+					ActualKey = Value;
+					ActualKey.Time = PreviousTime;
+					if (PreviousTime != NewTime)
+					{
+						RichCurve.SetKeyTime(KeyHandle, NewTime);
+					}
+
+					Model->OnCurveModified().Broadcast();
+				}
 			}
-		}
+		}	
 	}
 
 	virtual void UpdateValuesFromRawData() override
 	{
-		UObject* Owner = WeakOwner.Get();
-		if (Owner && RichCurve->IsKeyHandleValid(KeyHandle))
+		if(UObject* Owner = WeakOwner.Get())
 		{
-			Value = RichCurve->GetKey(KeyHandle);
+			if(Model->IsValid())
+			{
+				const FRichCurve& RichCurve = Model->GetReadOnlyRichCurve();
+				if (RichCurve.IsKeyHandleValid(KeyHandle))
+				{
+					Value = RichCurve.GetKey(KeyHandle);
+				}
+			}
 		}
 	}
 
@@ -77,8 +98,8 @@ private:
 
 	/** Cached key handle that this key proxy relates to */
 	FKeyHandle KeyHandle;
-	/** Cached curve in which the key resides */
-	FRichCurve* RichCurve;
+	/** Cached model that created this proxy */
+	FRichCurveEditorModel* Model;
 	/** Cached owner in which the raw curve resides */
 	TWeakObjectPtr<UObject> WeakOwner;
 };
\ No newline at end of file
diff --git a/Engine/Source/Editor/CurveEditor/Private/SCurveEditorPanel.cpp b/Engine/Source/Editor/CurveEditor/Private/SCurveEditorPanel.cpp
--- a/Engine/Source/Editor/CurveEditor/Private/SCurveEditorPanel.cpp
+++ b/Engine/Source/Editor/CurveEditor/Private/SCurveEditorPanel.cpp
@@ -232,13 +232,13 @@ void SCurveEditorPanel::Construct(const FArguments& InArgs, TSharedRef<FCurveEdi
 			.PhysicalSplitterHandleSize(2.0f)
 
 			+ SSplitter::Slot()
-			.Value(0.30f)
+			.Value(InArgs._TreeSplitterWidth)
 			[
 				InArgs._TreeContent.Widget
 			]
 
 			+ SSplitter::Slot()
-			.Value(0.7f)
+			.Value(InArgs._ContentSplitterWidth)
 			[
 				MainContent
 			]
@@ -684,12 +684,14 @@ void SCurveEditorPanel::UpdateEditBox()
 	for (TTuple<FCurveModelID, TMap<FKeyHandle, UObject*>>& OuterPair : EditObjects->CurveIDToKeyProxies)
 	{
 		const FKeyHandleSet* SelectedKeys = Selection.FindForCurve(OuterPair.Key);
-
-		for (TTuple<FKeyHandle, UObject*>& InnerPair : OuterPair.Value)
+		if(SelectedKeys)
 		{
-			if (ICurveEditorKeyProxy* Proxy = Cast<ICurveEditorKeyProxy>(InnerPair.Value))
+			for (TTuple<FKeyHandle, UObject*>& InnerPair : OuterPair.Value)
 			{
-				Proxy->UpdateValuesFromRawData();
+				if (ICurveEditorKeyProxy* Proxy = Cast<ICurveEditorKeyProxy>(InnerPair.Value))
+				{
+					Proxy->UpdateValuesFromRawData();
+				}
 			}
 		}
 	}
diff --git a/Engine/Source/Editor/CurveEditor/Private/Tree/CurveEditorTree.cpp b/Engine/Source/Editor/CurveEditor/Private/Tree/CurveEditorTree.cpp
--- a/Engine/Source/Editor/CurveEditor/Private/Tree/CurveEditorTree.cpp
+++ b/Engine/Source/Editor/CurveEditor/Private/Tree/CurveEditorTree.cpp
@@ -28,6 +28,17 @@ TArrayView<const FCurveModelID> FCurveEditorTreeItem::GetOrCreateCurves(FCurveEd
 			}
 		}
 	}
+	else
+	{
+		for (const FCurveModelID& ID : Curves)
+		{
+			FCurveModel* CurveModel = CurveEditor->FindCurve(ID);
+			if (CurveModel)
+			{
+				CurveEditor->BroadcastCurveChanged(CurveModel);
+			}
+		}
+	}
 	return Curves;
 }
 
@@ -429,6 +440,7 @@ void FCurveEditorTree::SetDirectSelection(TArray<FCurveEditorTreeItemID>&& TreeI
 		}
 	}
 
+
 	if (!PreviousSelection.OrderIndependentCompareEqual(Selection))
 	{
 		++Events.OnSelectionChanged.SerialNumber;
diff --git a/Engine/Source/Editor/CurveEditor/Private/Views/SInteractiveCurveEditorView.cpp b/Engine/Source/Editor/CurveEditor/Private/Views/SInteractiveCurveEditorView.cpp
--- a/Engine/Source/Editor/CurveEditor/Private/Views/SInteractiveCurveEditorView.cpp
+++ b/Engine/Source/Editor/CurveEditor/Private/Views/SInteractiveCurveEditorView.cpp
@@ -492,165 +492,6 @@ void SInteractiveCurveEditorView::Tick(const FGeometry& AllottedGeometry, const
 	GetCurveDrawParams(CachedDrawParams);
 }
 
-void SInteractiveCurveEditorView::GetCurveDrawParams(TArray<FCurveDrawParams>& OutDrawParams) const
-{
-	TSharedPtr<FCurveEditor> CurveEditor = WeakCurveEditor.Pin();
-	if (!CurveEditor)
-	{
-		return;
-	}
-
-	// Get the Min/Max values on the X axis, for Time
-	double InputMin = 0, InputMax = 1;
-	GetInputBounds(InputMin, InputMax);
-
-	ECurveEditorTangentVisibility TangentVisibility = CurveEditor->GetSettings()->GetTangentVisibility();
-	OutDrawParams.Reserve(CurveInfoByID.Num());
-
-	for (const TTuple<FCurveModelID, FCurveInfo>& Pair : CurveInfoByID)
-	{
-		FCurveModel* CurveModel = CurveEditor->FindCurve(Pair.Key);
-		if (!ensureAlways(CurveModel))
-		{
-			continue;
-		}
-
-		FCurveEditorScreenSpace CurveSpace = GetCurveSpace(Pair.Key);
-		
-		const float DisplayRatio = (CurveSpace.PixelsPerOutput() / CurveSpace.PixelsPerInput());
-
-		const FKeyHandleSet* SelectedKeys = CurveEditor->GetSelection().GetAll().Find(Pair.Key);
-
-		// Create a new set of Curve Drawing Parameters to represent this particular Curve
-		FCurveDrawParams Params(Pair.Key);
-		Params.Color = CurveModel->GetColor();
-		Params.bKeyDrawEnabled = CurveModel->IsKeyDrawEnabled();
-
-		// Gather the display metrics to use for each key type. This allows a Curve Model to override
-		// whether or not the curve supports Keys, Arrive/Leave Tangents, etc. If the Curve Model doesn't
-		// support a particular capability we can skip drawing them.
-		CurveModel->GetKeyDrawInfo(ECurvePointType::ArriveTangent, FKeyHandle::Invalid(), Params.ArriveTangentDrawInfo);
-		CurveModel->GetKeyDrawInfo(ECurvePointType::LeaveTangent, FKeyHandle::Invalid(), Params.LeaveTangentDrawInfo);
-
-		// Gather the interpolating points in input/output space
-		TArray<TTuple<double, double>> InterpolatingPoints;
-
-		CurveModel->DrawCurve(*CurveEditor, CurveSpace, InterpolatingPoints);
-		Params.InterpolatingPoints.Reserve(InterpolatingPoints.Num());
-
-		// An Input Offset allows for a fixed offset to all keys, such as displaying them in the middle of a frame instead of at the start.
-		double InputOffset = CurveModel->GetInputDisplayOffset();
-
-		// Convert the interpolating points to screen space
-		for (TTuple<double, double> Point : InterpolatingPoints)
-		{
-			Params.InterpolatingPoints.Add(
-				FVector2D(
-					CurveSpace.SecondsToScreen(Point.Get<0>() + InputOffset),
-					CurveSpace.ValueToScreen(Point.Get<1>())
-				)
-			);
-		}
-
-		TArray<FKeyHandle> VisibleKeys;
-		CurveModel->GetKeys(*CurveEditor, InputMin, InputMax, TNumericLimits<double>::Lowest(), TNumericLimits<double>::Max(), VisibleKeys);
-
-		if (VisibleKeys.Num())
-		{
-			TArray<FKeyPosition> AllKeyPositions;
-			TArray<FKeyAttributes> AllKeyAttributes;
-
-			AllKeyPositions.SetNum(VisibleKeys.Num());
-			AllKeyAttributes.SetNum(VisibleKeys.Num());
-
-			CurveModel->GetKeyPositions(VisibleKeys, AllKeyPositions);
-			CurveModel->GetKeyAttributes(VisibleKeys, AllKeyAttributes);
-
-			for (int32 Index = 0; Index < VisibleKeys.Num(); ++Index)
-			{
-				const FKeyHandle      KeyHandle = VisibleKeys[Index];
-				const FKeyPosition&   KeyPosition = AllKeyPositions[Index];
-				const FKeyAttributes& Attributes = AllKeyAttributes[Index];
-
-				bool bShowTangents = TangentVisibility == ECurveEditorTangentVisibility::AllTangents || (TangentVisibility == ECurveEditorTangentVisibility::SelectedKeys && SelectedKeys && SelectedKeys->Contains(VisibleKeys[Index]));
-
-				float TimeScreenPos = CurveSpace.SecondsToScreen(KeyPosition.InputValue + InputOffset);
-				float ValueScreenPos = CurveSpace.ValueToScreen(KeyPosition.OutputValue);
-
-				// Add this key
-				FCurvePointInfo Key(KeyHandle);
-				Key.ScreenPosition = FVector2D(TimeScreenPos, ValueScreenPos);
-				Key.LayerBias = 2;
-
-				// Add draw info for the specific key
-				CurveModel->GetKeyDrawInfo(ECurvePointType::Key, KeyHandle, /*Out*/ Key.DrawInfo);
-				Params.Points.Add(Key);
-
-				if (bShowTangents && Attributes.HasArriveTangent())
-				{
-					float ArriveTangent = Attributes.GetArriveTangent();
-
-					FCurvePointInfo ArriveTangentPoint(KeyHandle);
-					ArriveTangentPoint.Type = ECurvePointType::ArriveTangent;
-
-
-					if (Attributes.HasTangentWeightMode() && Attributes.HasArriveTangentWeight() &&
-						(Attributes.GetTangentWeightMode() == RCTWM_WeightedBoth || Attributes.GetTangentWeightMode() == RCTWM_WeightedArrive))
-					{
-						FVector2D TangentOffset = CurveEditor::ComputeScreenSpaceTangentOffset(CurveSpace, ArriveTangent, -Attributes.GetArriveTangentWeight());
-						ArriveTangentPoint.ScreenPosition = Key.ScreenPosition + TangentOffset;
-					}
-					else
-					{
-						float PixelLength = 60.0f;
-						ArriveTangentPoint.ScreenPosition = Key.ScreenPosition + CurveEditor::GetVectorFromSlopeAndLength(ArriveTangent * -DisplayRatio, -PixelLength);
-					}
-					ArriveTangentPoint.LineDelta = Key.ScreenPosition - ArriveTangentPoint.ScreenPosition;
-					ArriveTangentPoint.LayerBias = 1;
-
-					// Add draw info for the specific tangent
-					FKeyDrawInfo TangentDrawInfo;
-					CurveModel->GetKeyDrawInfo(ECurvePointType::ArriveTangent, KeyHandle, /*Out*/ ArriveTangentPoint.DrawInfo);
-
-					Params.Points.Add(ArriveTangentPoint);
-				}
-
-				if (bShowTangents && Attributes.HasLeaveTangent())
-				{
-					float LeaveTangent = Attributes.GetLeaveTangent();
-
-					FCurvePointInfo LeaveTangentPoint(KeyHandle);
-					LeaveTangentPoint.Type = ECurvePointType::LeaveTangent;
-
-					if (Attributes.HasTangentWeightMode() && Attributes.HasLeaveTangentWeight() &&
-						(Attributes.GetTangentWeightMode() == RCTWM_WeightedBoth || Attributes.GetTangentWeightMode() == RCTWM_WeightedLeave))
-					{
-						FVector2D TangentOffset = CurveEditor::ComputeScreenSpaceTangentOffset(CurveSpace, LeaveTangent, Attributes.GetLeaveTangentWeight());
-
-						LeaveTangentPoint.ScreenPosition = Key.ScreenPosition + TangentOffset;
-					}
-					else
-					{
-						float PixelLength = 60.0f;
-						LeaveTangentPoint.ScreenPosition = Key.ScreenPosition + CurveEditor::GetVectorFromSlopeAndLength(LeaveTangent * -DisplayRatio, PixelLength);
-
-					}
-
-					LeaveTangentPoint.LineDelta = Key.ScreenPosition - LeaveTangentPoint.ScreenPosition;
-					LeaveTangentPoint.LayerBias = 1;
-
-					// Add draw info for the specific tangent
-					FKeyDrawInfo TangentDrawInfo;
-					CurveModel->GetKeyDrawInfo(ECurvePointType::LeaveTangent, KeyHandle, /*Out*/ LeaveTangentPoint.DrawInfo);
-
-					Params.Points.Add(LeaveTangentPoint);
-				}
-			}
-		}
-
-		OutDrawParams.Add(MoveTemp(Params));
-	}
-}
 void SInteractiveCurveEditorView::GetPointsWithinWidgetRange(const FSlateRect& WidgetRectangle, TArray<FCurvePointHandle>* OutPoints) const
 {
 	TSharedPtr<FCurveEditor> CurveEditor = WeakCurveEditor.Pin();
diff --git a/Engine/Source/Editor/CurveEditor/Public/CurveDataAbstraction.h b/Engine/Source/Editor/CurveEditor/Public/CurveDataAbstraction.h
--- a/Engine/Source/Editor/CurveEditor/Public/CurveDataAbstraction.h
+++ b/Engine/Source/Editor/CurveEditor/Public/CurveDataAbstraction.h
@@ -173,10 +173,10 @@ private:
 
 	/** This key's entry tangent, if bHasArriveTangent */
 	UPROPERTY()
-	float ArriveTangent;
+	float ArriveTangent{};
 	/** This key's exit tangent, if bHasLeaveTangent */
 	UPROPERTY()
-	float LeaveTangent;
+	float LeaveTangent{};
 	/** This key's interpolation mode, if bHasInterpMode */
 	UPROPERTY()
 	TEnumAsByte<ERichCurveInterpMode> InterpMode;
@@ -188,10 +188,10 @@ private:
 	TEnumAsByte<ERichCurveTangentWeightMode> TangentWeightMode;
 	/** This key's entry tangent weight, if bHasArriveTangentWeight */
 	UPROPERTY()
-	float ArriveTangentWeight;
+	float ArriveTangentWeight{};
 	/** This key's exit tangent weight, if bHasLeaveTangentWeight */
 	UPROPERTY()
-	float LeaveTangentWeight;
+	float LeaveTangentWeight{};
 
 };
 
diff --git a/Engine/Source/Editor/CurveEditor/Public/CurveEditor.h b/Engine/Source/Editor/CurveEditor/Public/CurveEditor.h
--- a/Engine/Source/Editor/CurveEditor/Public/CurveEditor.h
+++ b/Engine/Source/Editor/CurveEditor/Public/CurveEditor.h
@@ -37,8 +37,8 @@ class ICurveEditorToolExtension;
 class IBufferedCurveModel;
 
 DECLARE_DELEGATE_OneParam(FOnSetBoolean, bool)
-DECLARE_MULTICAST_DELEGATE_OneParam(FOnCurveArrayChanged, FCurveModelID)
 DECLARE_MULTICAST_DELEGATE_OneParam(FOnActiveToolChanged, FCurveEditorToolID)
+DECLARE_MULTICAST_DELEGATE_TwoParams(FOnCurveArrayChanged, FCurveModel*, bool /*displayed*/);
 
 
 class CURVEEDITOR_API FCurveEditor 
@@ -78,6 +78,11 @@ public:
 	/** Grid line label text format strings for the X and Y axis */
 	TAttribute<FText> GridLineLabelFormatXAttribute, GridLineLabelFormatYAttribute;
 
+	/** Padding applied to zoom-to-fit the input */
+	TAttribute<double> InputZoomToFitPadding;
+
+	/** Padding applied to zoom-to-fit the output */
+	TAttribute<double> OutputZoomToFitPadding;
 
 	/** Delegate that is invoked when a tool becomes active. Also fired when the tool goes inactive. */
 	FOnActiveToolChanged OnActiveToolChangedDelegate;
@@ -104,6 +109,10 @@ public:
 
 	TSharedPtr<SCurveEditorPanel> GetPanel() const;
 
+	void SetView(TSharedPtr<SCurveEditorView> InPanel);
+
+	TSharedPtr<SCurveEditorView> GetView() const;
+
 	FCurveEditorScreenSpaceH GetPanelInputSpace() const;
 
 public:
@@ -193,6 +202,9 @@ public:
 	*/
 	void RemoveCurve(FCurveModelID InCurveID);
 
+	/** Remove all curves from this editor */
+	void RemoveAllCurves();
+
 	bool IsCurvePinned(FCurveModelID InCurveID) const;
 
 	void PinCurve(FCurveModelID InCurveID);
@@ -292,6 +304,16 @@ public:
 	 */
 	void RemoveTreeItem(FCurveEditorTreeItemID ItemID);
 
+	/**
+	 * Remove all tree items fromt he curve editor
+	 */
+	void RemoveAllTreeItems();
+
+	/**
+	 * Set the tree selection directly
+	 */
+	void SetTreeSelection(TArray<FCurveEditorTreeItemID>&& TreeItems);
+
 	/**
 	 * Check whether this tree item is selected
 	 */
@@ -395,6 +417,13 @@ public:
 	bool CanFlattenOrStraightenSelection() const;
 
 public:
+	/**
+	 * Populate the specified array with curve painting parameters
+	 *
+	 * @param OutDrawParams    An array to populate with curve painting parameters, one per visible curve
+	 */
+	void GetCurveDrawParams(TArray<FCurveDrawParams>& OutDrawParams) const;
+
 	/**
 	 * Called by SCurveEditorPanel to update the allocated geometry for this curve editor.
 	 */
@@ -418,6 +447,27 @@ public:
 	 */
 	void BindCommands();
 
+public:
+
+	/** Suspend or resume broadcast of curve array changing  */
+	void SuspendBroadcast()
+	{
+		SuspendBroadcastCount++;
+	}
+
+	void ResumeBroadcast()
+	{
+		SuspendBroadcastCount--;
+		checkf(SuspendBroadcastCount >= 0, TEXT("Suspend/Resume broadcast mismatch Curve Editor!"));
+	}
+
+	bool IsBroadcasting()
+	{
+		return SuspendBroadcastCount == 0;
+	}
+
+	void BroadcastCurveChanged(FCurveModel* InCurve);
+
 protected:
 
 	/**
@@ -454,6 +504,8 @@ protected:
 
 	TWeakPtr<SCurveEditorPanel> WeakPanel;
 
+	TWeakPtr<SCurveEditorView> WeakView;
+
 	/** Hierarchical information pertaining to curve data */
 	FCurveEditorTree Tree;
 
@@ -490,9 +542,18 @@ protected:
 	/** A serial number that is incremented any time the currently active set of curves are changed */
 	uint32 ActiveCurvesSerialNumber;
 
+	/** Counter to suspend broadcasting of changed delegates*/
+	int32 SuspendBroadcastCount;
+
 private:
 
 	/** Cached physical size of the panel representing this editor */
 	FVector2D CachedPhysicalSize;
 
+public:
+	/**
+	* Delegate that's broadcast when the curve display changes.
+	*/
+	FOnCurveArrayChanged OnCurveArrayChanged;
+
 };
diff --git a/Engine/Source/Editor/CurveEditor/Public/CurveEditorScreenSpace.h b/Engine/Source/Editor/CurveEditor/Public/CurveEditorScreenSpace.h
--- a/Engine/Source/Editor/CurveEditor/Public/CurveEditorScreenSpace.h
+++ b/Engine/Source/Editor/CurveEditor/Public/CurveEditorScreenSpace.h
@@ -20,7 +20,7 @@ struct FCurveEditorScreenSpaceH
 		: PixelWidth(FMath::Max(InPixelWidth, 1.f))
 		, InputMin(InInputMin), InputMax(InInputMax)
 	{
-		if (!ensure(InputMax > InputMin))
+		if (InputMax <= InputMin)
 		{
 			InputMax = InputMin + KINDA_SMALL_NUMBER;
 		}
diff --git a/Engine/Source/Editor/CurveEditor/Public/CurveModel.h b/Engine/Source/Editor/CurveEditor/Public/CurveModel.h
--- a/Engine/Source/Editor/CurveEditor/Public/CurveModel.h
+++ b/Engine/Source/Editor/CurveEditor/Public/CurveModel.h
@@ -9,6 +9,7 @@
 #include "Curves/RichCurve.h"
 #include "CurveEditorTypes.h"
 #include "Math/TransformCalculus2D.h"
+#include "Misc/Attribute.h"
 
 #include "IBufferedCurveModel.h"
 
@@ -38,7 +39,7 @@ public:
 
 	FCurveModel()
 		: Color(FLinearColor::White)
-		, bKeyDrawEnabled(1)
+		, bKeyDrawEnabled(true)
 		, SupportedViews(ECurveEditorViewID::ANY_BUILT_IN)
 	{}
 
@@ -232,7 +233,13 @@ public:
 	{
 		return false;
 	}
-
+	/**
+	* Get the UObject that owns this CurveModel, for example for Sequencer this would be the UMovieSceneSection
+	*/
+	virtual UObject* GetOwningObject() const
+	{
+		return nullptr;
+	}
 	/**
 	 * Helper function for assigning a the same attributes to a number of keys
 	 */
@@ -243,6 +250,15 @@ public:
 	 */
 	TOptional<FKeyHandle> AddKey(const FKeyPosition& NewKeyPosition, const FKeyAttributes& InAttributes);
 
+
+	/**
+	 * Get a multicast delegate, fired when modifications are made to this curve
+	 */
+	FORCEINLINE FSimpleMulticastDelegate& OnCurveModified()
+	{
+		return CurveModifiedDelegate;
+	}
+
 public:
 
 	/**
@@ -325,15 +341,15 @@ public:
  */
 	FORCEINLINE bool IsKeyDrawEnabled() const
 	{
-		return bKeyDrawEnabled != 0;
+		return bKeyDrawEnabled.Get();
 	}
 
 	/**
 	 * Assign whether or not to disable drawing keys
 	 */
-	FORCEINLINE void SetIsKeyDrawEnabled(bool bInKeyDrawEnabled)
+	FORCEINLINE void SetIsKeyDrawEnabled(TAttribute<bool> bInKeyDrawEnabled)
 	{
-		bKeyDrawEnabled = bInKeyDrawEnabled ? 1 : 0;
+		bKeyDrawEnabled = bInKeyDrawEnabled;
 	}
 
 	/**
@@ -359,8 +375,11 @@ protected:
 	FLinearColor Color;
 
 	/** Whether or not to draw curve's keys */
-	uint8 bKeyDrawEnabled : 1;
+	TAttribute<bool> bKeyDrawEnabled;
 
 	/** A set of views supported by this curve */
 	ECurveEditorViewID SupportedViews;
+
+	/** Multicast delegate broadcast on curve modification */
+	FSimpleMulticastDelegate CurveModifiedDelegate;
 };
diff --git a/Engine/Source/Editor/CurveEditor/Public/RichCurveEditorModel.h b/Engine/Source/Editor/CurveEditor/Public/RichCurveEditorModel.h
--- a/Engine/Source/Editor/CurveEditor/Public/RichCurveEditorModel.h
+++ b/Engine/Source/Editor/CurveEditor/Public/RichCurveEditorModel.h
@@ -13,7 +13,7 @@ struct FRichCurve;
 class CURVEEDITOR_API FRichCurveEditorModel : public FCurveModel
 {
 public:
-	FRichCurveEditorModel(FRichCurve* InRichCurve, UObject* InOwner);
+	FRichCurveEditorModel(UObject* InOwner);
 
 	virtual const void* GetCurve() const override;
 
@@ -43,9 +43,35 @@ public:
 
 	virtual TUniquePtr<IBufferedCurveModel> CreateBufferedCurveCopy() const override;
 
-private:
+	// Set a range to clamp key input values
+	void SetClampInputRange(TAttribute<TRange<double>> InClampInputRange) { ClampInputRange = InClampInputRange; }
 
-	FRichCurve* RichCurve;
+	// Check for whether this rich curve is valid. This is required mostly in the case of undo/redo, where 
+	// curves can potentially become invalid underneath this model before owners get their undo/redo callbacks
+	virtual bool IsValid() const = 0;
+
+	// Get the rich curve we are operating on
+	virtual FRichCurve& GetRichCurve() = 0;
+	virtual const FRichCurve& GetReadOnlyRichCurve() const = 0;
+
+private:
 
 	TWeakObjectPtr<> WeakOwner;
+
+	TAttribute<TRange<double>> ClampInputRange;
+};
+
+// Rich curve model operating on a raw curve ptr
+class CURVEEDITOR_API FRichCurveEditorModelRaw : public FRichCurveEditorModel
+{
+public:
+	FRichCurveEditorModelRaw(FRichCurve* InRichCurve, UObject* InOwner);
+
+	// FRichCurveEditorModel interface
+	virtual bool IsValid() const override { return RichCurve != nullptr; }
+	virtual FRichCurve& GetRichCurve() override;
+	virtual const FRichCurve& GetReadOnlyRichCurve() const override;
+
+private:
+	FRichCurve* RichCurve;
 };
\ No newline at end of file
diff --git a/Engine/Source/Editor/CurveEditor/Public/SCurveEditorPanel.h b/Engine/Source/Editor/CurveEditor/Public/SCurveEditorPanel.h
--- a/Engine/Source/Editor/CurveEditor/Public/SCurveEditorPanel.h
+++ b/Engine/Source/Editor/CurveEditor/Public/SCurveEditorPanel.h
@@ -29,6 +29,8 @@ class CURVEEDITOR_API SCurveEditorPanel : public SCompoundWidget
 {
 	SLATE_BEGIN_ARGS(SCurveEditorPanel)
 		: _GridLineTint(FLinearColor(0.1f, 0.1f, 0.1f, 1.f))
+		, _TreeSplitterWidth(0.3f)
+		, _ContentSplitterWidth(0.7f)
 	{}
 
 		/** Color to draw grid lines */
@@ -46,6 +48,12 @@ class CURVEEDITOR_API SCurveEditorPanel : public SCompoundWidget
 		/** Widget slot for the tree content */
 		SLATE_NAMED_SLOT(FArguments, TreeContent)
 
+		/** The width of the splitter slot for the tree */
+		SLATE_ARGUMENT(float, TreeSplitterWidth)
+
+		/** The width of the splitter slot for the main content */
+		SLATE_ARGUMENT(float, ContentSplitterWidth)
+
 	SLATE_END_ARGS()
 
 	SCurveEditorPanel();
diff --git a/Engine/Source/Editor/CurveEditor/Public/SCurveEditorView.h b/Engine/Source/Editor/CurveEditor/Public/SCurveEditorView.h
--- a/Engine/Source/Editor/CurveEditor/Public/SCurveEditorView.h
+++ b/Engine/Source/Editor/CurveEditor/Public/SCurveEditorView.h
@@ -13,6 +13,7 @@
 class FCurveEditor;
 class SCurveEditorPanel;
 class FSlateRect;
+struct FCurveDrawParams;
 
 /**
  * This is the base widget type for all views that exist on a curve editor panel. A view may contain 0 or more curves (stored in CurveInfoByID).
@@ -178,6 +179,9 @@ public:
 
 protected:
 
+	/** Gets info about the curves being drawn. Converts actual curves into an abstract series of lines/points/handles/etc. */
+	void GetCurveDrawParams(TArray<FCurveDrawParams>& OutDrawParams) const;
+
 	// ~SWidget interface
 	virtual FVector2D ComputeDesiredSize(float LayoutScaleMultiplier) const override;
 
diff --git a/Engine/Source/Editor/CurveEditor/Public/SCurveViewerPanel.h b/Engine/Source/Editor/CurveEditor/Public/SCurveViewerPanel.h
new file mode 100644
--- /dev/null
+++ b/Engine/Source/Editor/CurveEditor/Public/SCurveViewerPanel.h
@@ -0,0 +1,55 @@
+// Copyright Epic Games, Inc. All Rights Reserved.
+
+#pragma once
+
+#include "Widgets/DeclarativeSyntaxSupport.h"
+#include "Widgets/SLeafWidget.h"
+
+#include "CurveDrawInfo.h"
+#include "CurveEditor.h"
+#include "CurveEditorTypes.h"
+#include "SCurveEditorView.h"
+
+/**
+ * Curve viewer widget that reflects the state of an FCurveEditor
+ */
+class CURVEEDITOR_API SCurveViewerPanel : public SCurveEditorView
+{
+	SLATE_BEGIN_ARGS(SCurveViewerPanel) {}
+
+	SLATE_ATTRIBUTE(float, CurveThickness)
+
+	SLATE_END_ARGS()
+
+	/**
+	 * Construct a new curve editor panel widget
+	 */
+	void Construct(const FArguments& InArgs, TSharedRef<FCurveEditor> InCurveEditor);
+
+	/**
+	 * Access the draw parameters that this curve editor has cached for this frame
+	 */
+	const TArray<FCurveDrawParams>& GetCachedDrawParams() const
+	{
+		return CachedDrawParams;
+	}
+
+private:
+
+	virtual int32 OnPaint(const FPaintArgs& Args, const FGeometry& AllottedGeometry, const FSlateRect& MyCullingRect, FSlateWindowElementList& OutDrawElements, int32 LayerId, const FWidgetStyle& InWidgetStyle, bool bParentEnabled) const override;
+	virtual void Tick(const FGeometry& AllottedGeometry, const double InCurrentTime, const float InDeltaTime) override;
+
+public:
+
+	/**
+	 * Draw curve data
+	 */
+	void DrawCurves(const FGeometry& AllottedGeometry, FSlateWindowElementList& OutDrawElements, int32 LayerId, const FWidgetStyle& InWidgetStyle, ESlateDrawEffect DrawEffects) const;
+
+private:
+	/** Curve draw parameters that are re-generated on tick */
+	TArray<FCurveDrawParams> CachedDrawParams;
+
+	/** Thickness of the displayed curve */
+	TAttribute<float> CurveThickness;
+};
\ No newline at end of file
diff --git a/Engine/Source/Editor/CurveEditor/Public/Views/SInteractiveCurveEditorView.h b/Engine/Source/Editor/CurveEditor/Public/Views/SInteractiveCurveEditorView.h
--- a/Engine/Source/Editor/CurveEditor/Public/Views/SInteractiveCurveEditorView.h
+++ b/Engine/Source/Editor/CurveEditor/Public/Views/SInteractiveCurveEditorView.h
@@ -108,9 +108,6 @@ protected:
 
 private:
 
-	/** Gets info about the curves being drawn. Converts actual curves into an abstract series of lines/points/handles/etc. */
-	void GetCurveDrawParams(TArray<FCurveDrawParams>& OutDrawParams) const;
-
 	void CreateContextMenu(const FGeometry& MyGeometry, const FPointerEvent& MouseEvent);
 
 	/** Updates our distance to all of the curves we represent. */
diff --git a/Engine/Source/Editor/DetailCustomizations/Private/AnimSequenceDetails.cpp b/Engine/Source/Editor/DetailCustomizations/Private/AnimSequenceDetails.cpp
--- a/Engine/Source/Editor/DetailCustomizations/Private/AnimSequenceDetails.cpp
+++ b/Engine/Source/Editor/DetailCustomizations/Private/AnimSequenceDetails.cpp
@@ -15,7 +15,6 @@
 #include "AnimMontageSegmentDetails.h"
 #include "AnimPreviewInstance.h"
 #include "Slate/SceneViewport.h"
-#include "AnimationCompressionPanel.h"
 
 #define LOCTEXT_NAMESPACE	"AnimSequenceDetails"
 
@@ -161,38 +160,6 @@ void FAnimSequenceDetails::CustomizeDetails(IDetailLayoutBuilder& DetailBuilder)
 	];
 
 	CreateOverridenProperty(DetailBuilder, AdditiveSettingsCategory, RefFrameIndexHandle, TAttribute<EVisibility>( this, &FAnimSequenceDetails::ShouldShowRefFrameIndex ));
-
-	/////////////////////////////////////////////////////////////////
-	//  compression category!
-	///////////////////////////////////////////////////////////////
-	// add Apply button Compression
-	IDetailCategoryBuilder& CompressionCategory = DetailBuilder.EditCategory("Compression");
-
-	TArray<TSharedRef<IPropertyHandle>> CompressionProperties;
-
-	CompressionCategory.GetDefaultProperties( CompressionProperties );
-
-	for (auto& Property : CompressionProperties)
-	{
-		CompressionCategory.AddProperty(Property);
-	}
-
-	FDetailWidgetRow& CustomRow = CompressionCategory.AddCustomRow(LOCTEXT("ApplyCompressionLabel", "Apply"))
-		.WholeRowContent()
-		[
-			SNew(SVerticalBox)
-			+SVerticalBox::Slot()
-			.AutoHeight()
-			.Padding(2.0f, 0.0f)
-			.VAlign(VAlign_Center)
-			.HAlign(HAlign_Left)
-			[
-				SNew(SButton)
-				.Text(LOCTEXT("EditCompressionButton_Label", "Edit Compression Settings"))
-				.ToolTipText(LOCTEXT("EditCompressionButton_Tooltip", "Click to view and edit the Compression Settings"))
-				.OnClicked(this, &FAnimSequenceDetails::OnEditCompression)
-			]
-		];
 }
 END_SLATE_FUNCTION_BUILD_OPTIMIZATION
 
@@ -340,12 +307,6 @@ TSharedPtr<FString> FAnimSequenceDetails::GetRetargetSourceString(FName Retarget
 	return RetargetSourceComboList[0];
 }
 
-FReply FAnimSequenceDetails::OnEditCompression()
-{
-	FDlgAnimCompression AnimCompressionDialog(SelectedAnimSequences);
-	AnimCompressionDialog.ShowModal();
-	return FReply::Handled();
-}
 ////////////////////////////////////////////////
 // based on SAnimationSegmentViewport
 SAnimationRefPoseViewport::~SAnimationRefPoseViewport()
diff --git a/Engine/Source/Editor/DetailCustomizations/Private/AnimSequenceDetails.h b/Engine/Source/Editor/DetailCustomizations/Private/AnimSequenceDetails.h
--- a/Engine/Source/Editor/DetailCustomizations/Private/AnimSequenceDetails.h
+++ b/Engine/Source/Editor/DetailCustomizations/Private/AnimSequenceDetails.h
@@ -73,7 +73,6 @@ private:
 	// cache all anim sequences that are selected
 	// but I need to know them before compress
 	TArray< TWeakObjectPtr<UAnimSequence> > SelectedAnimSequences;
-	FReply OnEditCompression();
 };
 
 ///////////////////////////////////////////////////
diff --git a/Engine/Source/Editor/EditorStyle/Private/SlateEditorStyle.cpp b/Engine/Source/Editor/EditorStyle/Private/SlateEditorStyle.cpp
--- a/Engine/Source/Editor/EditorStyle/Private/SlateEditorStyle.cpp
+++ b/Engine/Source/Editor/EditorStyle/Private/SlateEditorStyle.cpp
@@ -2559,6 +2559,7 @@ void FSlateEditorStyle::FStyle::SetupSequencerStyles()
 		Set("Sequencer.AnimationOutliner.TransparentBorder", new FSlateColorBrush(FLinearColor::Transparent));
 		Set("Sequencer.AnimationOutliner.BoldFont", DEFAULT_FONT("Bold", 11));
 		Set("Sequencer.AnimationOutliner.RegularFont", DEFAULT_FONT("Regular", 9));
+		Set("Sequencer.AnimationOutliner.ItalicFont", DEFAULT_FONT("Italic", 10));
 		Set("Sequencer.ShotFilter", new IMAGE_BRUSH("Sequencer/FilteredArea", FVector2D(74, 74), FLinearColor::White, ESlateBrushTileType::Both));
 		Set("Sequencer.KeyMark", new IMAGE_BRUSH("Sequencer/KeyMark", FVector2D(3, 21), FLinearColor::White, ESlateBrushTileType::NoTile));
 		Set("Sequencer.ToggleAutoKeyEnabled", new IMAGE_BRUSH("Sequencer/Main_Icons/Icon_Sequencer_Auto_Key_24x", Icon48x48));
@@ -5781,6 +5782,22 @@ void FSlateEditorStyle::FStyle::SetupPersonaStyle()
 		Set("AnimationEditor.SetKey.Small", new IMAGE_BRUSH(TEXT("Icons/icon_Persona_SetKey_40x"), Icon20x20));
 		Set("AnimationEditor.ApplyAnimation", new IMAGE_BRUSH(TEXT("Icons/icon_Persona_BakeAnim_40x"), Icon40x40));
 		Set("AnimationEditor.ApplyAnimation.Small", new IMAGE_BRUSH(TEXT("Icons/icon_Persona_BakeAnim_40x"), Icon20x20));
+
+		Set("AnimTimeline.Outliner.DefaultBorder", new FSlateColorBrush(FLinearColor::White));
+
+		const FSplitterStyle OutlinerSplitterStyle = FSplitterStyle()
+			.SetHandleNormalBrush( FSlateNoResource() )
+			.SetHandleHighlightBrush( FSlateNoResource() );
+		Set("AnimTimeline.Outliner.Splitter", OutlinerSplitterStyle);
+
+		Set("AnimTimeline.Outliner.Label", FTextBlockStyle(NormalText)
+			.SetColorAndOpacity(FLinearColor(1.0f, 1.0f, 1.0f))
+			.SetShadowOffset(FVector2D(1.0f, 1.0f)));
+
+		Set("AnimTimeline.Outliner.ItemColor", FLinearColor(0.04f, 0.04f, 0.04f, 0.8f));
+		Set("AnimTimeline.Outliner.HeaderColor", FLinearColor(0.03f, 0.03f, 0.03f, 1.0f));
+
+		Set("AnimTimeline.SectionMarker", new IMAGE_BRUSH(TEXT("Sequencer/SectionMarker"), FVector2D(11, 12)));
 	}
 
 	// Skeletal mesh editor
diff --git a/Engine/Source/Editor/EditorWidgets/Private/SAssetSearchBox.cpp b/Engine/Source/Editor/EditorWidgets/Private/SAssetSearchBox.cpp
--- a/Engine/Source/Editor/EditorWidgets/Private/SAssetSearchBox.cpp
+++ b/Engine/Source/Editor/EditorWidgets/Private/SAssetSearchBox.cpp
@@ -348,7 +348,7 @@ void SAssetSearchBox::UpdateSuggestionList()
 		}
 	}
 
-	if (Suggestions.Num() > 0)
+	if (Suggestions.Num() > 0 && HasKeyboardFocus())
 	{
 		// At least one suggestion was found, open the menu
 		SuggestionBox->SetIsOpen(true, false);
diff --git a/Engine/Source/Editor/GraphEditor/Private/SGraphNode.cpp b/Engine/Source/Editor/GraphEditor/Private/SGraphNode.cpp
--- a/Engine/Source/Editor/GraphEditor/Private/SGraphNode.cpp
+++ b/Engine/Source/Editor/GraphEditor/Private/SGraphNode.cpp
@@ -948,30 +948,39 @@ void SGraphNode::UpdateGraphNode()
 			]			
 		];
 
-	// Create comment bubble
-	TSharedPtr<SCommentBubble> CommentBubble;
-	const FSlateColor CommentColor = GetDefault<UGraphEditorSettings>()->DefaultCommentNodeTitleColor;
-
-	SAssignNew( CommentBubble, SCommentBubble )
-	.GraphNode( GraphNode )
-	.Text( this, &SGraphNode::GetNodeComment )
-	.OnTextCommitted( this, &SGraphNode::OnCommentTextCommitted )
-	.OnToggled( this, &SGraphNode::OnCommentBubbleToggled )
-	.ColorAndOpacity( CommentColor )
-	.AllowPinning( true )
-	.EnableTitleBarBubble( true )
-	.EnableBubbleCtrls( true )
-	.GraphLOD( this, &SGraphNode::GetCurrentLOD )
-	.IsGraphNodeHovered( this, &SGraphNode::IsHovered );
-
-	GetOrAddSlot( ENodeZone::TopCenter )
-	.SlotOffset( TAttribute<FVector2D>( CommentBubble.Get(), &SCommentBubble::GetOffset ))
-	.SlotSize( TAttribute<FVector2D>( CommentBubble.Get(), &SCommentBubble::GetSize ))
-	.AllowScaling( TAttribute<bool>( CommentBubble.Get(), &SCommentBubble::IsScalingAllowed ))
-	.VAlign( VAlign_Top )
-	[
-		CommentBubble.ToSharedRef()
-	];
+	bool SupportsBubble = true;
+	if (GraphNode != nullptr)
+	{
+		SupportsBubble = GraphNode->SupportsCommentBubble();
+	}
+
+	if (SupportsBubble)
+	{
+		// Create comment bubble
+		TSharedPtr<SCommentBubble> CommentBubble;
+		const FSlateColor CommentColor = GetDefault<UGraphEditorSettings>()->DefaultCommentNodeTitleColor;
+
+		SAssignNew(CommentBubble, SCommentBubble)
+			.GraphNode(GraphNode)
+			.Text(this, &SGraphNode::GetNodeComment)
+			.OnTextCommitted(this, &SGraphNode::OnCommentTextCommitted)
+			.OnToggled(this, &SGraphNode::OnCommentBubbleToggled)
+			.ColorAndOpacity(CommentColor)
+			.AllowPinning(true)
+			.EnableTitleBarBubble(true)
+			.EnableBubbleCtrls(true)
+			.GraphLOD(this, &SGraphNode::GetCurrentLOD)
+			.IsGraphNodeHovered(this, &SGraphNode::IsHovered);
+
+		GetOrAddSlot(ENodeZone::TopCenter)
+			.SlotOffset(TAttribute<FVector2D>(CommentBubble.Get(), &SCommentBubble::GetOffset))
+			.SlotSize(TAttribute<FVector2D>(CommentBubble.Get(), &SCommentBubble::GetSize))
+			.AllowScaling(TAttribute<bool>(CommentBubble.Get(), &SCommentBubble::IsScalingAllowed))
+			.VAlign(VAlign_Top)
+			[
+				CommentBubble.ToSharedRef()
+			];
+	}
 
 	CreateBelowWidgetControls(MainVerticalBox);
 	CreatePinWidgets();
diff --git a/Engine/Source/Editor/GraphEditor/Public/SGraphNode.h b/Engine/Source/Editor/GraphEditor/Public/SGraphNode.h
--- a/Engine/Source/Editor/GraphEditor/Public/SGraphNode.h
+++ b/Engine/Source/Editor/GraphEditor/Public/SGraphNode.h
@@ -239,7 +239,7 @@ public:
 	EGraphRenderingLOD::Type GetCurrentLOD() const;
 
 	/** Called when GraphNode changes its error information, may be called when no change has actually occurred: */
-	void RefreshErrorInfo();
+	virtual void RefreshErrorInfo();
 	
 	/** Called when ed graph data is cleared, indicating this widget can no longer safely access GraphNode - forwards call to owned pins: */
 	void InvalidateGraphData();
@@ -299,7 +299,7 @@ protected:
 	FString GetErrorMessage() const {return ErrorMsg;}
 
 	/** Called to set error text on the node */
-	void UpdateErrorInfo();
+	virtual void UpdateErrorInfo();
 
 	/** Set-up the error reporting widget for the node */
 	void SetupErrorReporting();
diff --git a/Engine/Source/Editor/Kismet/Private/BlueprintActionMenuItem.cpp b/Engine/Source/Editor/Kismet/Private/BlueprintActionMenuItem.cpp
--- a/Engine/Source/Editor/Kismet/Private/BlueprintActionMenuItem.cpp
+++ b/Engine/Source/Editor/Kismet/Private/BlueprintActionMenuItem.cpp
@@ -67,7 +67,12 @@ static void FBlueprintMenuActionItemImpl::DirtyBlueprintFromNewNode(UEdGraphNode
 	
 	UBlueprint* Blueprint = FBlueprintEditorUtils::FindBlueprintForGraphChecked(NodeGraph);
 	check(Blueprint != nullptr);
-	
+
+	if (SpawnedNode->GetSchema()->MarkBlueprintDirtyFromNewNode(Blueprint, SpawnedNode))
+	{
+		return;
+	}
+
 	UK2Node* K2Node = Cast<UK2Node>(SpawnedNode);
 	// see if we need to recompile skeleton after adding this node, or just mark
 	// it dirty (default to rebuilding the skel, since there is no way to if
diff --git a/Engine/Source/Editor/Kismet/Private/BlueprintEditor.cpp b/Engine/Source/Editor/Kismet/Private/BlueprintEditor.cpp
--- a/Engine/Source/Editor/Kismet/Private/BlueprintEditor.cpp
+++ b/Engine/Source/Editor/Kismet/Private/BlueprintEditor.cpp
@@ -3096,6 +3096,16 @@ bool FBlueprintEditor::CanNavigateToChildGraph() const
 	return FocusedGraphEdPtr.IsValid() && (FocusedGraphEdPtr.Pin()->GetCurrentGraph()->SubGraphs.Num() > 0);
 }
 
+bool FBlueprintEditor::TransactionObjectAffectsBlueprint(UObject* InTransactedObject)
+{
+	check(InTransactedObject);
+
+	UBlueprint* BlueprintObj = GetBlueprintObj();
+	check(BlueprintObj);
+
+	return InTransactedObject->GetOutermost() == BlueprintObj->GetOutermost();
+}
+
 void FBlueprintEditor::HandleUndoTransaction(const FTransaction* Transaction)
 {
 	UBlueprint* BlueprintObj = GetBlueprintObj();
@@ -3107,9 +3117,10 @@ void FBlueprintEditor::HandleUndoTransaction(const FTransaction* Transaction)
 		// Look at the transaction this function is responding to, see if any object in it has an outermost of the Blueprint
 		TArray<UObject*> TransactionObjects;
 		Transaction->GetTransactionObjects(TransactionObjects);
+
 		for (UObject* Object : TransactionObjects)
 		{
-			if (Object->GetOutermost() == BlueprintOutermost)
+			if(TransactionObjectAffectsBlueprint(Object))
 			{
 				bAffectsBlueprint = true;
 				break;
diff --git a/Engine/Source/Editor/Kismet/Private/BlueprintEditorModule.cpp b/Engine/Source/Editor/Kismet/Private/BlueprintEditorModule.cpp
--- a/Engine/Source/Editor/Kismet/Private/BlueprintEditorModule.cpp
+++ b/Engine/Source/Editor/Kismet/Private/BlueprintEditorModule.cpp
@@ -130,7 +130,10 @@ void FixSubObjectReferencesPostUndoRedo(const FTransaction* Transaction)
 
 			if (Blueprint != nullptr)
 			{
-				ModifiedBlueprints.AddUnique(Blueprint);
+				if (Blueprint->ShouldBeMarkedDirtyUponTransaction())
+				{
+					ModifiedBlueprints.AddUnique(Blueprint);
+				}
 			}
 		}
 	}
@@ -140,7 +143,7 @@ void FixSubObjectReferencesPostUndoRedo(const FTransaction* Transaction)
 	{
 		FixSubObjectReferencesPostUndoRedo(Blueprint->GeneratedClass->GetDefaultObject());
 		// Will cause a call to RefreshEditors()
-		if (Blueprint->RequiresMarkAsStructurallyModifiedOnUndo())
+		if (Blueprint->ShouldBeMarkedDirtyUponTransaction())
 		{
 			FBlueprintEditorUtils::MarkBlueprintAsStructurallyModified(Blueprint);
 		}
diff --git a/Engine/Source/Editor/Kismet/Public/BlueprintEditor.h b/Engine/Source/Editor/Kismet/Public/BlueprintEditor.h
--- a/Engine/Source/Editor/Kismet/Public/BlueprintEditor.h
+++ b/Engine/Source/Editor/Kismet/Public/BlueprintEditor.h
@@ -877,7 +877,7 @@ protected:
 	bool CanSelectAllNodes() const;
 
 	virtual void DeleteSelectedNodes();
-	bool CanDeleteNodes() const;
+	virtual bool CanDeleteNodes() const;
 
 	/**
 	* Given a node, make connections from anything connected to it's input pin to
@@ -889,16 +889,16 @@ protected:
 
 	void DeleteSelectedDuplicatableNodes();
 
-	void CutSelectedNodes();
-	bool CanCutNodes() const;
+	virtual void CutSelectedNodes();
+	virtual bool CanCutNodes() const;
 
-	void CopySelectedNodes();
-	bool CanCopyNodes() const;
+	virtual void CopySelectedNodes();
+	virtual bool CanCopyNodes() const;
 
 	/** Paste on graph at specific location */
 	virtual void PasteNodesHere(class UEdGraph* DestinationGraph, const FVector2D& GraphLocation) override;
 
-	void PasteNodes();
+	virtual void PasteNodes();
 	virtual bool CanPasteNodes() const override;
 
 	void DuplicateNodes();
@@ -910,11 +910,11 @@ protected:
 	void OnAssignReferencedActor();
 	bool CanAssignReferencedActor() const;
 
-	void OnStartWatchingPin();
-	bool CanStartWatchingPin() const;
+	virtual void OnStartWatchingPin();
+	virtual bool CanStartWatchingPin() const;
 
-	void OnStopWatchingPin();
-	bool CanStopWatchingPin() const;
+	virtual void OnStopWatchingPin();
+	virtual bool CanStopWatchingPin() const;
 
 	/**  BEGIN PERSONA related callback functions */
 	virtual void OnSelectBone() {};
@@ -1177,6 +1177,9 @@ private:
 	void HandleUndoTransaction(const class FTransaction* Transaction);
 
 public://@TODO
+
+	virtual bool TransactionObjectAffectsBlueprint(UObject* InTransactedObject);
+
 	TSharedPtr<FDocumentTracker> DocumentManager;
 	
 	/** Update all nodes' unrelated states when the graph has changed */
diff --git a/Engine/Source/Editor/LevelEditor/LevelEditor.Build.cs b/Engine/Source/Editor/LevelEditor/LevelEditor.Build.cs
--- a/Engine/Source/Editor/LevelEditor/LevelEditor.Build.cs
+++ b/Engine/Source/Editor/LevelEditor/LevelEditor.Build.cs
@@ -15,8 +15,9 @@ public class LevelEditor : ModuleRules
 				"SlateReflector",
                 "IntroTutorials",
                 "AppFramework",
-                "PortalServices"
-			}
+                "PortalServices",
+                "Persona",
+            }
 		);
 
 		PublicIncludePathModuleNames.AddRange(
@@ -57,7 +58,6 @@ public class LevelEditor : ModuleRules
 				"TargetPlatform",
 				"EngineSettings",
 				"PropertyEditor",
-				"Persona",
 				"Kismet",
 				"KismetWidgets",
 				"Sequencer",
@@ -90,8 +90,9 @@ public class LevelEditor : ModuleRules
                 "PlacementMode",
                 "IntroTutorials",
 				"HeadMountedDisplay",
-				"VREditor"
-			}
+				"VREditor",
+                "Persona",
+            }
 		);
 
 		if(Target.bWithLiveCoding)
diff --git a/Engine/Source/Editor/MaterialEditor/Private/MaterialStatsCommon.cpp b/Engine/Source/Editor/MaterialEditor/Private/MaterialStatsCommon.cpp
--- a/Engine/Source/Editor/MaterialEditor/Private/MaterialStatsCommon.cpp
+++ b/Engine/Source/Editor/MaterialEditor/Private/MaterialStatsCommon.cpp
@@ -351,7 +351,7 @@ void FMaterialStatsUtils::GetRepresentativeShaderTypesAndDescriptions(TMap<FName
 	bool bMobileHDR = MobileHDR && MobileHDR->GetValueOnAnyThread() == 1;
 
 	static const FName FLocalVertexFactoryName = FLocalVertexFactory::StaticType.GetFName();
-	static const FName FGPUFactoryName = TGPUSkinVertexFactory<true>::StaticType.GetFName();
+	static const FName FGPUFactoryName = TEXT("TGPUSkinVertexFactoryExtra");
 
 	if (TargetMaterial->IsUIMaterial())
 	{
diff --git a/Engine/Source/Editor/MovieSceneTools/Private/Channels/BuiltInChannelEditors.cpp b/Engine/Source/Editor/MovieSceneTools/Private/Channels/BuiltInChannelEditors.cpp
--- a/Engine/Source/Editor/MovieSceneTools/Private/Channels/BuiltInChannelEditors.cpp
+++ b/Engine/Source/Editor/MovieSceneTools/Private/Channels/BuiltInChannelEditors.cpp
@@ -20,7 +20,7 @@
 #include "Framework/MultiBox/MultiBoxBuilder.h"
 #include "Channels/MovieSceneChannelProxy.h"
 #include "Channels/MovieSceneChannelEditorData.h"
-#include "FloatChannelCurveModel.h"
+#include "Channels/FloatChannelCurveModel.h"
 #include "EventChannelCurveModel.h"
 #include "PropertyCustomizationHelpers.h"
 
diff --git a/Engine/Source/Editor/MovieSceneTools/Private/Channels/FloatChannelCurveModel.cpp b/Engine/Source/Editor/MovieSceneTools/Private/Channels/FloatChannelCurveModel.cpp
--- a/Engine/Source/Editor/MovieSceneTools/Private/Channels/FloatChannelCurveModel.cpp
+++ b/Engine/Source/Editor/MovieSceneTools/Private/Channels/FloatChannelCurveModel.cpp
@@ -1,6 +1,6 @@
 // Copyright Epic Games, Inc. All Rights Reserved.
 
-#include "FloatChannelCurveModel.h"
+#include "Channels/FloatChannelCurveModel.h"
 #include "Math/Vector2D.h"
 #include "HAL/PlatformMath.h"
 #include "Channels/MovieSceneFloatChannel.h"
@@ -119,6 +119,8 @@ void FFloatChannelCurveModel::AddKeys(TArrayView<const FKeyPosition> InKeyPositi
 		// We reuse SetKeyAttributes here as there is complex logic determining which parts of the attributes are valid to set.
 		// For now we need to duplicate the new key handle array due to API mismatch. This will auto calculate tangents if needed.
 		SetKeyAttributes(NewKeyHandles, InKeyAttributes);
+
+		CurveModifiedDelegate.Broadcast();
 	}
 }
 
@@ -160,6 +162,8 @@ void FFloatChannelCurveModel::RemoveKeys(TArrayView<const FKeyHandle> InKeys)
 				ChannelData.RemoveKey(KeyIndex);
 			}
 		}
+
+		CurveModifiedDelegate.Broadcast();
 	}
 }
 
@@ -326,6 +330,8 @@ void FFloatChannelCurveModel::SetKeyPositions(TArrayView<const FKeyHandle> InKey
 		}
 
 		Channel->AutoSetTangents();
+
+		CurveModifiedDelegate.Broadcast();
 	}
 }
 
@@ -513,6 +519,8 @@ void FFloatChannelCurveModel::SetKeyAttributes(TArrayView<const FKeyHandle> InKe
 		{
 			Channel->AutoSetTangents();
 		}
+
+		CurveModifiedDelegate.Broadcast();
 	}
 }
 
@@ -543,6 +551,8 @@ void FFloatChannelCurveModel::SetCurveAttributes(const FCurveAttributes& InCurve
 		{
 			Channel->PostInfinityExtrap = InCurveAttributes.GetPostExtrapolation();
 		}
+
+		CurveModifiedDelegate.Broadcast();
 	}
 }
 
diff --git a/Engine/Source/Editor/MovieSceneTools/Private/MovieSceneToolHelpers.cpp b/Engine/Source/Editor/MovieSceneTools/Private/MovieSceneToolHelpers.cpp
--- a/Engine/Source/Editor/MovieSceneTools/Private/MovieSceneToolHelpers.cpp
+++ b/Engine/Source/Editor/MovieSceneTools/Private/MovieSceneToolHelpers.cpp
@@ -60,6 +60,8 @@
 #include "FbxExporter.h"
 #include "Serialization/ObjectWriter.h"
 #include "Serialization/ObjectReader.h"
+#include "AnimationRecorder.h"
+#include "Components/SkeletalMeshComponent.h"
 
 /* MovieSceneToolHelpers
  *****************************************************************************/
@@ -2043,4 +2045,72 @@ bool MovieSceneToolHelpers::ExportFBX(UWorld* World, UMovieScene* MovieScene, IM
 	Exporter->WriteToFile(*InFBXFileName);
 
 	return true;
-}
\ No newline at end of file
+}
+
+
+
+bool MovieSceneToolHelpers::ExportToAnimSequence(UAnimSequence* AnimSequence, UMovieScene* MovieScene, IMovieScenePlayer* Player,
+	USkeletalMeshComponent* SkelMeshComp, FMovieSceneSequenceIDRef& Template, FMovieSceneSequenceTransform& RootToLocalTransform)
+{
+	//if we have no allocated bone space transforms something wrong so try to recalc them
+	if (SkelMeshComp->GetBoneSpaceTransforms().Num() <= 0)
+	{
+		SkelMeshComp->RecalcRequiredBones(0);
+		if (SkelMeshComp->GetBoneSpaceTransforms().Num() <= 0)
+		{
+
+			UE_LOG(LogMovieScene, Error, TEXT("Error Animation Anim Sequence Export, No Bone Transforms."));
+			return false;
+		}
+	}
+
+	UnFbx::FLevelSequenceAnimTrackAdapter AnimTrackAdapter(Player, MovieScene, RootToLocalTransform);
+
+	FFrameRate SampleRate = MovieScene->GetDisplayRate();
+
+	FAnimRecorderInstance AnimationRecorder;
+	FAnimationRecordingSettings RecordingSettings;
+	RecordingSettings.SampleRate = SampleRate.AsDecimal();
+	RecordingSettings.InterpMode = ERichCurveInterpMode::RCIM_Cubic;
+	RecordingSettings.TangentMode = ERichCurveTangentMode::RCTM_Auto;
+	RecordingSettings.Length = 0;
+	RecordingSettings.bRecordInWorldSpace = true;
+	RecordingSettings.bRemoveRootAnimation = false;
+	RecordingSettings.bCheckDeltaTimeAtBeginning = false;
+	AnimationRecorder.Init(SkelMeshComp, AnimSequence, nullptr, RecordingSettings);
+	AnimationRecorder.BeginRecording();
+
+
+	int32 LocalStartFrame = AnimTrackAdapter.GetLocalStartFrame();
+	int32 StartFrame = AnimTrackAdapter.GetStartFrame();
+	int32 AnimationLength = AnimTrackAdapter.GetLength();
+	float FrameRate = AnimTrackAdapter.GetFrameRate();
+	float DeltaTime = 1.0f / FrameRate;
+	for (int32 FrameCount = 0; FrameCount <= AnimationLength; ++FrameCount)
+	{
+
+		int32 LocalFrame = LocalStartFrame + FrameCount;
+
+		// This will call UpdateSkelPose on the skeletal mesh component to move bones based on animations in the matinee group
+		AnimTrackAdapter.UpdateAnimation(LocalFrame);
+
+		// Update space bases so new animation position has an effect.
+		// @todo - hack - this will be removed at some point (this comment is all over the place by the way in fbx export code).
+		SkelMeshComp->TickAnimation(0.03f, false);
+
+		SkelMeshComp->RefreshBoneTransforms();
+		SkelMeshComp->RefreshSlaveComponents();
+		SkelMeshComp->UpdateComponentToWorld();
+		SkelMeshComp->FinalizeBoneTransform();
+		SkelMeshComp->MarkRenderTransformDirty();
+		SkelMeshComp->MarkRenderDynamicDataDirty();
+
+		AnimationRecorder.Update(DeltaTime);
+	}
+
+	const bool bShowAnimationAssetCreatedToast = false;
+	AnimationRecorder.FinishRecording(bShowAnimationAssetCreatedToast);
+
+	return true;
+}
+
