fix CIS build warnings #rb none

diff --git a/Engine/Source/Editor/CurveEditor/Public/CurveDataAbstraction.h b/Engine/Source/Editor/CurveEditor/Public/CurveDataAbstraction.h
--- a/Engine/Source/Editor/CurveEditor/Public/CurveDataAbstraction.h
+++ b/Engine/Source/Editor/CurveEditor/Public/CurveDataAbstraction.h
@@ -51,6 +51,13 @@ struct FKeyAttributes
 		bHasTangentWeightMode = 0;
 		bHasArriveTangentWeight = 0;
 		bHasLeaveTangentWeight = 0;
+		ArriveTangent = 0.0f;
+		LeaveTangent = 0.0f;
+		InterpMode = RCIM_Linear;
+		TangentMode = RCTM_Auto;
+		TangentWeightMode = RCTWM_WeightedNone;
+		ArriveTangentWeight = 0.0f;
+		LeaveTangentWeight = 0.0f;
 	}
 
 	/**
