Fixed robomerge collision.

diff --git a/Engine/Source/Editor/CurveEditor/Private/RichCurveEditorModel.cpp b/Engine/Source/Editor/CurveEditor/Private/RichCurveEditorModel.cpp
--- a/Engine/Source/Editor/CurveEditor/Private/RichCurveEditorModel.cpp
+++ b/Engine/Source/Editor/CurveEditor/Private/RichCurveEditorModel.cpp
@@ -13,6 +13,93 @@
 #include "EditorStyleSet.h"
 #include "UObject/Package.h"
 
+
+void RefineCurvePoints(const FRichCurve* RichCurve, double TimeThreshold, float ValueThreshold, TArray<TTuple<double, double>>& InOutPoints)
+{
+	check(RichCurve);
+	const float InterpTimes[] = { 0.25f, 0.5f, 0.6f };
+
+	for (int32 Index = 0; Index < InOutPoints.Num() - 1; ++Index)
+	{
+		TTuple<double, double> Lower = InOutPoints[Index];
+		TTuple<double, double> Upper = InOutPoints[Index + 1];
+
+		if ((Upper.Get<0>() - Lower.Get<0>()) >= TimeThreshold)
+		{
+			bool bSegmentIsLinear = true;
+
+			TTuple<double, double> Evaluated[ARRAY_COUNT(InterpTimes)];
+
+			for (int32 InterpIndex = 0; InterpIndex < ARRAY_COUNT(InterpTimes); ++InterpIndex)
+			{
+				double& EvalTime = Evaluated[InterpIndex].Get<0>();
+
+				EvalTime = FMath::Lerp(Lower.Get<0>(), Upper.Get<0>(), InterpTimes[InterpIndex]);
+
+				float Value = RichCurve->Eval(EvalTime);
+
+				const float LinearValue = FMath::Lerp(Lower.Get<1>(), Upper.Get<1>(), InterpTimes[InterpIndex]);
+				if (bSegmentIsLinear)
+				{
+					bSegmentIsLinear = FMath::IsNearlyEqual(Value, LinearValue, ValueThreshold);
+				}
+
+				Evaluated[InterpIndex].Get<1>() = Value;
+			}
+
+			if (!bSegmentIsLinear)
+			{
+				// Add the point
+				InOutPoints.Insert(Evaluated, ARRAY_COUNT(Evaluated), Index + 1);
+				--Index;
+			}
+		}
+	}
+}
+
+/**
+ * Buffered curve implementation for a rich curve, stores a copy of the rich curve in order to draw itself.
+ */
+class FRichBufferedCurveModel : public IBufferedCurveModel
+{
+public:
+	FRichBufferedCurveModel(FRichCurve* InRichCurve, TArray<FKeyPosition>&& InKeyPositions, TArray<FKeyAttributes>&& InKeyAttributes,
+		const FString& InIntentionName, const double InValueMin, const double InValueMax)
+		: IBufferedCurveModel(MoveTemp(InKeyPositions), MoveTemp(InKeyAttributes), InIntentionName, InValueMin, InValueMax)
+		, RichCurve(*InRichCurve)
+	{}
+
+	virtual void DrawCurve(const FCurveEditor& CurveEditor, const FCurveEditorScreenSpace& ScreenSpace, TArray<TTuple<double, double>>& InterpolatingPoints) const override
+	{
+		const double StartTimeSeconds = ScreenSpace.GetInputMin();
+		const double EndTimeSeconds = ScreenSpace.GetInputMax();
+		const double TimeThreshold = FMath::Max(0.0001, 1.0 / ScreenSpace.PixelsPerInput());
+		const double ValueThreshold = FMath::Max(0.0001, 1.0 / ScreenSpace.PixelsPerOutput());
+
+		InterpolatingPoints.Add(MakeTuple(StartTimeSeconds, double(RichCurve.Eval(StartTimeSeconds))));
+
+		for (const FRichCurveKey& Key : RichCurve.GetConstRefOfKeys())
+		{
+			if (Key.Time > StartTimeSeconds && Key.Time < EndTimeSeconds)
+			{
+				InterpolatingPoints.Add(MakeTuple(double(Key.Time), double(Key.Value)));
+			}
+		}
+
+		InterpolatingPoints.Add(MakeTuple(EndTimeSeconds, double(RichCurve.Eval(EndTimeSeconds))));
+
+		int32 OldSize = InterpolatingPoints.Num();
+		do
+		{
+			OldSize = InterpolatingPoints.Num();
+			RefineCurvePoints(&RichCurve, TimeThreshold, ValueThreshold, InterpolatingPoints);
+		} while (OldSize != InterpolatingPoints.Num());
+	}
+
+private:
+	FRichCurve RichCurve;
+};
+
 FRichCurveEditorModel::FRichCurveEditorModel(FRichCurve* InRichCurve, UObject* InOwner)
 	: RichCurve(InRichCurve), WeakOwner(InOwner)
 {
@@ -89,48 +176,6 @@ void FRichCurveEditorModel::RemoveKeys(TArrayView<const FKeyHandle> InKeys)
 	}
 }
 
-void RefineCurvePoints(const FRichCurve* RichCurve, double TimeThreshold, float ValueThreshold, TArray<TTuple<double, double>>& InOutPoints)
-{
-	const float InterpTimes[] = { 0.25f, 0.5f, 0.6f };
-
-	for (int32 Index = 0; Index < InOutPoints.Num() - 1; ++Index)
-	{
-		TTuple<double, double> Lower = InOutPoints[Index];
-		TTuple<double, double> Upper = InOutPoints[Index + 1];
-
-		if ((Upper.Get<0>() - Lower.Get<0>()) >= TimeThreshold)
-		{
-			bool bSegmentIsLinear = true;
-
-			TTuple<double, double> Evaluated[UE_ARRAY_COUNT(InterpTimes)];
-
-			for (int32 InterpIndex = 0; InterpIndex < UE_ARRAY_COUNT(InterpTimes); ++InterpIndex)
-			{
-				double& EvalTime  = Evaluated[InterpIndex].Get<0>();
-
-				EvalTime = FMath::Lerp(Lower.Get<0>(), Upper.Get<0>(), InterpTimes[InterpIndex]);
-
-				float Value = RichCurve->Eval(EvalTime);
-
-				const float LinearValue = FMath::Lerp(Lower.Get<1>(), Upper.Get<1>(), InterpTimes[InterpIndex]);
-				if (bSegmentIsLinear)
-				{
-					bSegmentIsLinear = FMath::IsNearlyEqual(Value, LinearValue, ValueThreshold);
-				}
-
-				Evaluated[InterpIndex].Get<1>() = Value;
-			}
-
-			if (!bSegmentIsLinear)
-			{
-				// Add the point
-				InOutPoints.Insert(Evaluated, UE_ARRAY_COUNT(Evaluated), Index+1);
-				--Index;
-			}
-		}
-	}
-}
-
 void FRichCurveEditorModel::DrawCurve(const FCurveEditor& CurveEditor, const FCurveEditorScreenSpace& ScreenSpace, TArray<TTuple<double, double>>& InOutPoints) const
 {
 	if (UObject* Owner = WeakOwner.Get())
@@ -471,6 +516,30 @@ void FRichCurveEditorModel::CreateKeyProxies(TArrayView<const FKeyHandle> InKeyH
 	}
 }
 
+TUniquePtr<IBufferedCurveModel> FRichCurveEditorModel::CreateBufferedCurveCopy() const
+{
+	if (UObject* Owner = WeakOwner.Get())
+	{
+		TArray<FKeyHandle> TargetKeyHandles;
+		for (auto It = RichCurve->GetKeyHandleIterator(); It; ++It)
+		{
+			TargetKeyHandles.Add(*It);
+		}
+
+		TArray<FKeyPosition> KeyPositions;
+		TArray<FKeyAttributes> KeyAttributes;
+		GetKeyPositions(TargetKeyHandles, KeyPositions);
+		GetKeyAttributes(TargetKeyHandles, KeyAttributes);
+
+		double ValueMin = 0.f, ValueMax = 1.f;
+		GetValueRange(ValueMin, ValueMax);
+
+		return MakeUnique<FRichBufferedCurveModel>(RichCurve, MoveTemp(KeyPositions), MoveTemp(KeyAttributes), GetIntentionName(), ValueMin, ValueMax);
+	}
+
+	return nullptr;
+}
+
 void FRichCurveEditorModel::GetTimeRange(double& MinTime, double& MaxTime) const
 {
 	if (UObject* Owner = WeakOwner.Get())
@@ -521,4 +590,4 @@ void FRichCurveEditorModel::GetNeighboringKeys(const FKeyHandle InKeyHandle, TOp
 			}
 		}
 	}
-}
\ No newline at end of file
+}
