Fixed deprecation warning.

diff --git a/Engine/Source/Editor/CurveEditor/Private/RichCurveEditorModel.cpp b/Engine/Source/Editor/CurveEditor/Private/RichCurveEditorModel.cpp
--- a/Engine/Source/Editor/CurveEditor/Private/RichCurveEditorModel.cpp
+++ b/Engine/Source/Editor/CurveEditor/Private/RichCurveEditorModel.cpp
@@ -28,9 +28,9 @@ void RefineCurvePoints(const FRichCurve* RichCurve, double TimeThreshold, float
 		{
 			bool bSegmentIsLinear = true;
 
-			TTuple<double, double> Evaluated[ARRAY_COUNT(InterpTimes)];
+			TTuple<double, double> Evaluated[UE_ARRAY_COUNT(InterpTimes)];
 
-			for (int32 InterpIndex = 0; InterpIndex < ARRAY_COUNT(InterpTimes); ++InterpIndex)
+			for (int32 InterpIndex = 0; InterpIndex < UE_ARRAY_COUNT(InterpTimes); ++InterpIndex)
 			{
 				double& EvalTime = Evaluated[InterpIndex].Get<0>();
 
@@ -50,7 +50,7 @@ void RefineCurvePoints(const FRichCurve* RichCurve, double TimeThreshold, float
 			if (!bSegmentIsLinear)
 			{
 				// Add the point
-				InOutPoints.Insert(Evaluated, ARRAY_COUNT(Evaluated), Index + 1);
+				InOutPoints.Insert(Evaluated, UE_ARRAY_COUNT(Evaluated), Index + 1);
 				--Index;
 			}
 		}
