CurveEditor: Speculative fix for crash in curve editor when widgets inexplicably disappear.

diff --git a/Engine/Source/Editor/CurveEditor/Private/CurveEditor.cpp b/Engine/Source/Editor/CurveEditor/Private/CurveEditor.cpp
--- a/Engine/Source/Editor/CurveEditor/Private/CurveEditor.cpp
+++ b/Engine/Source/Editor/CurveEditor/Private/CurveEditor.cpp
@@ -253,6 +253,16 @@ const FCurveEditorTreeItem& FCurveEditor::GetTreeItem(FCurveEditorTreeItemID Ite
 	return Tree.GetItem(ItemID);
 }
 
+FCurveEditorTreeItem* FCurveEditor::FindTreeItem(FCurveEditorTreeItemID ItemID)
+{
+	return Tree.FindItem(ItemID);
+}
+
+const FCurveEditorTreeItem* FCurveEditor::FindTreeItem(FCurveEditorTreeItemID ItemID) const
+{
+	return Tree.FindItem(ItemID);
+}
+
 const TArray<FCurveEditorTreeItemID>& FCurveEditor::GetRootTreeItems() const
 {
 	return Tree.GetRootItems();
diff --git a/Engine/Source/Editor/CurveEditor/Private/Tree/SCurveEditorTreePin.cpp b/Engine/Source/Editor/CurveEditor/Private/Tree/SCurveEditorTreePin.cpp
--- a/Engine/Source/Editor/CurveEditor/Private/Tree/SCurveEditorTreePin.cpp
+++ b/Engine/Source/Editor/CurveEditor/Private/Tree/SCurveEditorTreePin.cpp
@@ -49,16 +49,18 @@ FReply SCurveEditorTreePin::TogglePinned()
 
 void SCurveEditorTreePin::PinRecursive(FCurveEditorTreeItemID InTreeItem, FCurveEditor* CurveEditor) const
 {
-	FCurveEditorTreeItem& Item = CurveEditor->GetTreeItem(InTreeItem);
-
-	for (FCurveModelID CurveID : Item.GetOrCreateCurves(CurveEditor))
+	FCurveEditorTreeItem* Item = CurveEditor->FindTreeItem(InTreeItem);
+	if (ensureMsgf(Item != nullptr, TEXT("Can't find curve editor tree item. Ignoring pinning request.")))
 	{
-		CurveEditor->PinCurve(CurveID);
-	}
+		for (FCurveModelID CurveID : Item->GetOrCreateCurves(CurveEditor))
+		{
+			CurveEditor->PinCurve(CurveID);
+		}
 
-	for (FCurveEditorTreeItemID Child : Item.GetChildren())
-	{
-		PinRecursive(Child, CurveEditor);
+		for (FCurveEditorTreeItemID Child : Item->GetChildren())
+		{
+			PinRecursive(Child, CurveEditor);
+		}
 	}
 }
 
@@ -66,31 +68,38 @@ void SCurveEditorTreePin::UnpinRecursive(FCurveEditorTreeItemID InTreeItem, FCur
 {
 	const bool bIsSelected = CurveEditor->GetTreeSelectionState(InTreeItem) == ECurveEditorTreeSelectionState::Explicit;
 
-	FCurveEditorTreeItem& Item = CurveEditor->GetTreeItem(InTreeItem);
-	for (FCurveModelID CurveID : Item.GetCurves())
+	FCurveEditorTreeItem* Item = CurveEditor->FindTreeItem(InTreeItem);
+	if (ensureMsgf(Item != nullptr, TEXT("Can't find curve editor tree item. Ignoring unpinning request.")))
 	{
-		if (bIsSelected)
+		for (FCurveModelID CurveID : Item->GetCurves())
 		{
-			CurveEditor->UnpinCurve(CurveID);
+			if (bIsSelected)
+			{
+				CurveEditor->UnpinCurve(CurveID);
+			}
+			else
+			{
+				Item->DestroyCurves(CurveEditor);
+			}
 		}
-		else
+
+		for (FCurveEditorTreeItemID Child : Item->GetChildren())
 		{
-			Item.DestroyCurves(CurveEditor);
+			UnpinRecursive(Child, CurveEditor);
 		}
 	}
-
-	for (FCurveEditorTreeItemID Child : Item.GetChildren())
-	{
-		UnpinRecursive(Child, CurveEditor);
-	}
 }
 
 bool SCurveEditorTreePin::IsPinnedRecursive(FCurveEditorTreeItemID InTreeItem, FCurveEditor* CurveEditor) const
 {
-	const FCurveEditorTreeItem& Item = CurveEditor->GetTreeItem(InTreeItem);
+	const FCurveEditorTreeItem* Item = CurveEditor->FindTreeItem(InTreeItem);
+	if (!ensureMsgf(Item != nullptr, TEXT("Can't find curve editor item. Acting like it's not pinned.")))
+	{
+		return false;
+	}
 
-	TArrayView<const FCurveModelID>          Curves   = Item.GetCurves();
-	TArrayView<const FCurveEditorTreeItemID> Children = Item.GetChildren();
+	TArrayView<const FCurveModelID>          Curves   = Item->GetCurves();
+	TArrayView<const FCurveEditorTreeItemID> Children = Item->GetChildren();
 
 	const bool bAllChildren = Algo::AllOf(Children, [this, CurveEditor](FCurveEditorTreeItemID In){ return this->IsPinnedRecursive(In, CurveEditor); });
 
@@ -132,4 +141,4 @@ const FSlateBrush* SCurveEditorTreePin::GetPinBrush() const
 	}
 
 	return FEditorStyle::GetBrush("GenericCurveEditor.Pin_Inactive");
-}
\ No newline at end of file
+}
diff --git a/Engine/Source/Editor/CurveEditor/Public/CurveEditor.h b/Engine/Source/Editor/CurveEditor/Public/CurveEditor.h
--- a/Engine/Source/Editor/CurveEditor/Public/CurveEditor.h
+++ b/Engine/Source/Editor/CurveEditor/Public/CurveEditor.h
@@ -289,6 +289,16 @@ public:
 	 */
 	const FCurveEditorTreeItem& GetTreeItem(FCurveEditorTreeItemID ItemID) const;
 
+	/**
+	 * Finds a tree item from its ID
+	 */
+	FCurveEditorTreeItem* FindTreeItem(FCurveEditorTreeItemID ItemID);
+
+	/**
+	 * Finds a tree item from its ID
+	 */
+	const FCurveEditorTreeItem* FindTreeItem(FCurveEditorTreeItemID ItemID) const;
+
 	/**
 	 * Get const access to the entire set of root tree items
 	 */
