Fixing merge issue that removed passing the change type through

diff --git a/Engine/Source/Editor/CurveEditor/Private/RichCurveEditorModel.cpp b/Engine/Source/Editor/CurveEditor/Private/RichCurveEditorModel.cpp
--- a/Engine/Source/Editor/CurveEditor/Private/RichCurveEditorModel.cpp
+++ b/Engine/Source/Editor/CurveEditor/Private/RichCurveEditorModel.cpp
@@ -333,7 +333,7 @@ void FRichCurveEditorModel::SetKeyPositions(TArrayView<const FKeyHandle> InKeys,
 				}
 			}
 			RichCurve.AutoSetTangents();
-			FPropertyChangedEvent PropertyChangeStruct(nullptr, EPropertyChangeType::ValueSet);
+			FPropertyChangedEvent PropertyChangeStruct(nullptr, ChangeType);
 			Owner->PostEditChangeProperty(PropertyChangeStruct);
 
 			CurveModifiedDelegate.Broadcast();
@@ -535,7 +535,7 @@ void FRichCurveEditorModel::SetKeyAttributes(TArrayView<const FKeyHandle> InKeys
 				RichCurve.AutoSetTangents();
 			}
 
-			FPropertyChangedEvent PropertyChangeStruct(nullptr, EPropertyChangeType::ValueSet);
+			FPropertyChangedEvent PropertyChangeStruct(nullptr, ChangeType);
 			Owner->PostEditChangeProperty(PropertyChangeStruct);
 			CurveModifiedDelegate.Broadcast();
 		}
