Sequencer: Step to next/previous key should wrap around

#rnx

diff --git a/Engine/Source/Editor/CurveEditor/Private/CurveEditor.cpp b/Engine/Source/Editor/CurveEditor/Private/CurveEditor.cpp
--- a/Engine/Source/Editor/CurveEditor/Private/CurveEditor.cpp
+++ b/Engine/Source/Editor/CurveEditor/Private/CurveEditor.cpp
@@ -615,9 +615,10 @@ void FCurveEditor::StepToNextKey()
 
 	FFrameRate TickResolution = WeakTimeSliderController.Pin()->GetTickResolution();
 
-	double CurrentTime = WeakTimeSliderController.Pin()->GetScrubPosition() / TickResolution;
+	double CurrentTime = TickResolution.AsSeconds(WeakTimeSliderController.Pin()->GetScrubPosition());
 
 	TOptional<double> NextTime;
+	TOptional<double> MinTime;
 
 	for (const TTuple<FCurveModelID, TUniquePtr<FCurveModel>>& Pair : CurveData)
 	{
@@ -643,6 +644,13 @@ void FCurveEditor::StepToNextKey()
 					}
 				}
 			}
+
+			double CurveMinTime, CurveMaxTime;
+			CurveModel->GetTimeRange(CurveMinTime, CurveMaxTime);
+			if (!MinTime.IsSet() || CurveMinTime < MinTime.GetValue())
+			{
+				MinTime = CurveMinTime;
+			}
 		}
 	}
 
@@ -650,6 +658,10 @@ void FCurveEditor::StepToNextKey()
 	{
 		WeakTimeSliderController.Pin()->SetScrubPosition(NextTime.GetValue() * TickResolution);
 	}
+	else if (MinTime.IsSet())
+	{
+		WeakTimeSliderController.Pin()->SetScrubPosition(MinTime.GetValue() * TickResolution);
+	}
 }
 
 void FCurveEditor::StepToPreviousKey()
@@ -661,9 +673,10 @@ void FCurveEditor::StepToPreviousKey()
 
 	FFrameRate TickResolution = WeakTimeSliderController.Pin()->GetTickResolution();
 
-	double CurrentTime = WeakTimeSliderController.Pin()->GetScrubPosition() / TickResolution;
+	double CurrentTime = TickResolution.AsSeconds(WeakTimeSliderController.Pin()->GetScrubPosition());
 
 	TOptional<double> PreviousTime;
+	TOptional<double> MaxTime;
 
 	for (const TTuple<FCurveModelID, TUniquePtr<FCurveModel>>& Pair : CurveData)
 	{
@@ -689,6 +702,13 @@ void FCurveEditor::StepToPreviousKey()
 					}
 				}
 			}
+
+			double CurveMinTime, CurveMaxTime;
+			CurveModel->GetTimeRange(CurveMinTime, CurveMaxTime);
+			if (!MaxTime.IsSet() || CurveMaxTime > MaxTime.GetValue())
+			{
+				MaxTime = CurveMaxTime;
+			}
 		}
 	}
 
@@ -696,6 +716,10 @@ void FCurveEditor::StepToPreviousKey()
 	{
 		WeakTimeSliderController.Pin()->SetScrubPosition(PreviousTime.GetValue() * TickResolution);
 	}
+	else if (MaxTime.IsSet())
+	{
+		WeakTimeSliderController.Pin()->SetScrubPosition(MaxTime.GetValue() * TickResolution);
+	}
 }
 
 
diff --git a/Engine/Source/Editor/Sequencer/Private/SSequencer.cpp b/Engine/Source/Editor/Sequencer/Private/SSequencer.cpp
--- a/Engine/Source/Editor/Sequencer/Private/SSequencer.cpp
+++ b/Engine/Source/Editor/Sequencer/Private/SSequencer.cpp
@@ -3077,6 +3077,20 @@ void SSequencer::StepToKey(bool bStepToNextKey, bool bCameraOnly)
 						}
 					}
 				}
+
+				if (!StepToTime.IsSet() && AllTimes.Num() > 0)
+				{
+					AllTimes.Sort();
+
+					if (bStepToNextKey)
+					{
+						StepToTime = AllTimes[0];
+					}
+					else
+					{
+						StepToTime = AllTimes.Last();
+					}
+				}
 			}
 
 			if ( StepToTime.IsSet() )
diff --git a/Engine/Source/Editor/Sequencer/Private/SequencerKeyCollection.cpp b/Engine/Source/Editor/Sequencer/Private/SequencerKeyCollection.cpp
--- a/Engine/Source/Editor/Sequencer/Private/SequencerKeyCollection.cpp
+++ b/Engine/Source/Editor/Sequencer/Private/SequencerKeyCollection.cpp
@@ -283,8 +283,17 @@ TOptional<FFrameNumber> FSequencerKeyCollection::GetNextKey(FFrameNumber FrameNu
 	{
 		return GroupedTimes[Index];
 	}
-	else
+	else if (GroupedTimes.Num() > 0)
 	{
-		return TOptional<FFrameNumber>();
+		if (Direction == EFindKeyDirection::Forwards)
+		{
+			return GroupedTimes[0];
+		}
+		else
+		{
+			return GroupedTimes.Last();
+		}
 	}
+
+	return TOptional<FFrameNumber>();
 }
