Curve Editor: Change FCurveMode::AddKey implementations to call UpdateOrAddKey instead of AddKey so that duplicate keys can't be added.

#rnx

diff --git a/Engine/Source/Editor/CurveEditor/Private/RichCurveEditorModel.cpp b/Engine/Source/Editor/CurveEditor/Private/RichCurveEditorModel.cpp
--- a/Engine/Source/Editor/CurveEditor/Private/RichCurveEditorModel.cpp
+++ b/Engine/Source/Editor/CurveEditor/Private/RichCurveEditorModel.cpp
@@ -146,12 +146,14 @@ void FRichCurveEditorModel::AddKeys(TArrayView<const FKeyPosition> InKeyPosition
 				FKeyAttributes Attributes = InKeyAttributes[Index];
 
 				const TRange<double> InputRange = ClampInputRange.Get();
-				FKeyHandle     NewHandle = RichCurve.AddKey(FMath::Clamp(Position.InputValue, InputRange.GetLowerBoundValue(), InputRange.GetUpperBoundValue()), Position.OutputValue);
-
-				NewKeyHandles[Index] = NewHandle;
-				if (OutKeyHandles)
+				FKeyHandle     NewHandle = RichCurve.UpdateOrAddKey(FMath::Clamp(Position.InputValue, InputRange.GetLowerBoundValue(), InputRange.GetUpperBoundValue()), Position.OutputValue);
+				if (NewHandle != FKeyHandle::Invalid())
 				{
-					(*OutKeyHandles)[Index] = NewHandle;
+					NewKeyHandles[Index] = NewHandle;
+					if (OutKeyHandles)
+					{
+						(*OutKeyHandles)[Index] = NewHandle;
+					}
 				}
 			}
 
diff --git a/Engine/Source/Editor/MovieSceneTools/Private/Channels/EventChannelCurveModel.cpp b/Engine/Source/Editor/MovieSceneTools/Private/Channels/EventChannelCurveModel.cpp
--- a/Engine/Source/Editor/MovieSceneTools/Private/Channels/EventChannelCurveModel.cpp
+++ b/Engine/Source/Editor/MovieSceneTools/Private/Channels/EventChannelCurveModel.cpp
@@ -64,10 +64,13 @@ void FEventChannelCurveModel::AddKeys(TArrayView<const FKeyPosition> InKeyPositi
 			Section->ExpandToFrame(Time);
 
 			FMovieSceneEvent Value;
-			int32 KeyIndex = ChannelData.AddKey(Time, Value);
-			if (OutKeyHandles)
+			FKeyHandle NewHandle = ChannelData.UpdateOrAddKey(Time, Value);
+			if (NewHandle != FKeyHandle::Invalid())
 			{
-				(*OutKeyHandles)[Index] = ChannelData.GetHandle(KeyIndex);
+				if (OutKeyHandles)
+				{
+					(*OutKeyHandles)[Index] = NewHandle;
+				}
 			}
 		}
 	}
diff --git a/Engine/Source/Editor/MovieSceneTools/Private/Channels/FloatChannelCurveModel.cpp b/Engine/Source/Editor/MovieSceneTools/Private/Channels/FloatChannelCurveModel.cpp
--- a/Engine/Source/Editor/MovieSceneTools/Private/Channels/FloatChannelCurveModel.cpp
+++ b/Engine/Source/Editor/MovieSceneTools/Private/Channels/FloatChannelCurveModel.cpp
@@ -106,13 +106,15 @@ void FFloatChannelCurveModel::AddKeys(TArrayView<const FKeyPosition> InKeyPositi
 
 			FMovieSceneFloatValue Value(Position.OutputValue);
 
-			int32 KeyIndex = ChannelData.AddKey(Time, Value);
-			FKeyHandle NewHandle = ChannelData.GetHandle(KeyIndex);
-			NewKeyHandles[Index] = NewHandle;
-
-			if (OutKeyHandles)
+			FKeyHandle NewHandle = ChannelData.UpdateOrAddKey(Time, Value);
+			if (NewHandle != FKeyHandle::Invalid())
 			{
-				(*OutKeyHandles)[Index] = NewHandle;
+				NewKeyHandles[Index] = NewHandle;
+
+				if (OutKeyHandles)
+				{
+					(*OutKeyHandles)[Index] = NewHandle;
+				}
 			}
 		}
 
