Don't post actions to right-click menu for curve editor if curve/key's parent curve is set to read only

diff --git a/Engine/Source/Editor/CurveEditor/Private/CurveEditorContextMenu.cpp b/Engine/Source/Editor/CurveEditor/Private/CurveEditorContextMenu.cpp
--- a/Engine/Source/Editor/CurveEditor/Private/CurveEditorContextMenu.cpp
+++ b/Engine/Source/Editor/CurveEditor/Private/CurveEditorContextMenu.cpp
@@ -19,6 +19,8 @@ void FCurveEditorContextMenu::BuildMenu(FMenuBuilder& MenuBuilder, TSharedRef<FC
 		return FText::Format(LOCTEXT("ApplyStoredCurvesContextMenu", "Apply {0} Stored Curves"), LocalCurveEditor->GetNumBufferedCurves());
 	}));
 
+	const FCurveModel* HoveredCurve = HoveredCurveID.IsSet() ? CurveEditor->FindCurve(HoveredCurveID.GetValue()) : nullptr;
+
 	// We prioritize key selections over curve selections to reduce the pixel-perfectness needed
 	// to edit the keys (which is more common than curves). Right clicking on a key or an empty space
 	// should show the key menu, otherwise we show the curve menu (ie: right clicking on a curve, not 
@@ -27,23 +29,42 @@ void FCurveEditorContextMenu::BuildMenu(FMenuBuilder& MenuBuilder, TSharedRef<FC
 	{
 		MenuBuilder.BeginSection("CurveEditorKeySection", FText::Format(LOCTEXT("CurveEditorKeySection", "{0} Selected {0}|plural(one=Key,other=Keys)"), NumSelectedKeys));
 		{
-			// Modify Data
-			MenuBuilder.AddMenuEntry(FGenericCommands::Get().Delete);
-			
-			MenuBuilder.AddMenuEntry(FCurveEditorCommands::Get().FlattenTangents);
-			MenuBuilder.AddMenuEntry(FCurveEditorCommands::Get().StraightenTangents);
+			bool bIsReadOnly = false;
+			if (HoveredCurve)
+			{
+				bIsReadOnly = HoveredCurve->IsReadOnly();
+			}
+			else
+			{
+				if (ClickedPoint.IsSet())
+				{
+					if (FCurveModel* ClickedPointCurve = CurveEditor->FindCurve(ClickedPoint->CurveID))
+					{
+						bIsReadOnly = ClickedPointCurve->IsReadOnly();
+					}
+				}
+			}
+
+			if (!bIsReadOnly)
+			{
+				// Modify Data
+				MenuBuilder.AddMenuEntry(FGenericCommands::Get().Delete);
+
+				MenuBuilder.AddMenuEntry(FCurveEditorCommands::Get().FlattenTangents);
+				MenuBuilder.AddMenuEntry(FCurveEditorCommands::Get().StraightenTangents);
 
-			MenuBuilder.AddMenuSeparator();
+				MenuBuilder.AddMenuSeparator();
 
-			// Tangent Types
-			MenuBuilder.AddMenuEntry(FCurveEditorCommands::Get().InterpolationCubicAuto);
-			MenuBuilder.AddMenuEntry(FCurveEditorCommands::Get().InterpolationCubicUser);
-			MenuBuilder.AddMenuEntry(FCurveEditorCommands::Get().InterpolationCubicBreak);
-			MenuBuilder.AddMenuEntry(FCurveEditorCommands::Get().InterpolationLinear);
-			MenuBuilder.AddMenuEntry(FCurveEditorCommands::Get().InterpolationConstant);
-			MenuBuilder.AddMenuEntry(FCurveEditorCommands::Get().InterpolationToggleWeighted);
+				// Tangent Types
+				MenuBuilder.AddMenuEntry(FCurveEditorCommands::Get().InterpolationCubicAuto);
+				MenuBuilder.AddMenuEntry(FCurveEditorCommands::Get().InterpolationCubicUser);
+				MenuBuilder.AddMenuEntry(FCurveEditorCommands::Get().InterpolationCubicBreak);
+				MenuBuilder.AddMenuEntry(FCurveEditorCommands::Get().InterpolationLinear);
+				MenuBuilder.AddMenuEntry(FCurveEditorCommands::Get().InterpolationConstant);
+				MenuBuilder.AddMenuEntry(FCurveEditorCommands::Get().InterpolationToggleWeighted);
 
-			MenuBuilder.AddMenuSeparator();
+				MenuBuilder.AddMenuSeparator();
+			}
 
 			// Filters
 			MenuBuilder.AddMenuEntry(FCurveEditorCommands::Get().OpenUserImplementableFilterWindow);
@@ -52,42 +73,47 @@ void FCurveEditorContextMenu::BuildMenu(FMenuBuilder& MenuBuilder, TSharedRef<FC
 	}
 	else
 	{
-		const FCurveModel* HoveredCurve = HoveredCurveID.IsSet() ? CurveEditor->FindCurve(HoveredCurveID.GetValue()) : nullptr;
 		if (HoveredCurve)
 		{
 			MenuBuilder.BeginSection("CurveEditorCurveSection", FText::Format(LOCTEXT("CurveNameFormat", "Curve '{0}'"), HoveredCurve->GetLongDisplayName()));
 			{
 				// Buffer Curves
 				MenuBuilder.AddMenuEntry(FCurveEditorCommands::Get().BufferVisibleCurves);
-				MenuBuilder.AddMenuEntry(FCurveEditorCommands::Get().ApplyBufferedCurves, NAME_None, ApplyBufferedCurvesText);
+				if (!HoveredCurve->IsReadOnly())
+				{
+					MenuBuilder.AddMenuEntry(FCurveEditorCommands::Get().ApplyBufferedCurves, NAME_None, ApplyBufferedCurvesText);
+				}
 				MenuBuilder.AddMenuSeparator();
 
 				// Modify Curve
-				MenuBuilder.AddMenuEntry(FCurveEditorCommands::Get().AddKeyHovered);
-
-				MenuBuilder.AddSubMenu(LOCTEXT("PreExtrapText", "Pre-Extrap"), FText(), FNewMenuDelegate::CreateLambda(
-					[](FMenuBuilder& SubMenu)
-					{
-						SubMenu.AddMenuEntry(FCurveEditorCommands::Get().SetPreInfinityExtrapCycle);
-						SubMenu.AddMenuEntry(FCurveEditorCommands::Get().SetPreInfinityExtrapCycleWithOffset);
-						SubMenu.AddMenuEntry(FCurveEditorCommands::Get().SetPreInfinityExtrapOscillate);
-						SubMenu.AddMenuEntry(FCurveEditorCommands::Get().SetPreInfinityExtrapLinear);
-						SubMenu.AddMenuEntry(FCurveEditorCommands::Get().SetPreInfinityExtrapConstant);
-					})
-				);
-
-				MenuBuilder.AddSubMenu(LOCTEXT("PostExtrapText", "Post-Extrap"), FText(), FNewMenuDelegate::CreateLambda(
-					[](FMenuBuilder& SubMenu)
-					{
-						SubMenu.AddMenuEntry(FCurveEditorCommands::Get().SetPostInfinityExtrapCycle);
-						SubMenu.AddMenuEntry(FCurveEditorCommands::Get().SetPostInfinityExtrapCycleWithOffset);
-						SubMenu.AddMenuEntry(FCurveEditorCommands::Get().SetPostInfinityExtrapOscillate);
-						SubMenu.AddMenuEntry(FCurveEditorCommands::Get().SetPostInfinityExtrapLinear);
-						SubMenu.AddMenuEntry(FCurveEditorCommands::Get().SetPostInfinityExtrapConstant);
-					})
-				);
-
-				MenuBuilder.AddMenuSeparator();
+				if (!HoveredCurve->IsReadOnly())
+				{
+					MenuBuilder.AddMenuEntry(FCurveEditorCommands::Get().AddKeyHovered);
+
+					MenuBuilder.AddSubMenu(LOCTEXT("PreExtrapText", "Pre-Extrap"), FText(), FNewMenuDelegate::CreateLambda(
+						[](FMenuBuilder& SubMenu)
+						{
+							SubMenu.AddMenuEntry(FCurveEditorCommands::Get().SetPreInfinityExtrapCycle);
+							SubMenu.AddMenuEntry(FCurveEditorCommands::Get().SetPreInfinityExtrapCycleWithOffset);
+							SubMenu.AddMenuEntry(FCurveEditorCommands::Get().SetPreInfinityExtrapOscillate);
+							SubMenu.AddMenuEntry(FCurveEditorCommands::Get().SetPreInfinityExtrapLinear);
+							SubMenu.AddMenuEntry(FCurveEditorCommands::Get().SetPreInfinityExtrapConstant);
+						})
+					);
+
+					MenuBuilder.AddSubMenu(LOCTEXT("PostExtrapText", "Post-Extrap"), FText(), FNewMenuDelegate::CreateLambda(
+						[](FMenuBuilder& SubMenu)
+						{
+							SubMenu.AddMenuEntry(FCurveEditorCommands::Get().SetPostInfinityExtrapCycle);
+							SubMenu.AddMenuEntry(FCurveEditorCommands::Get().SetPostInfinityExtrapCycleWithOffset);
+							SubMenu.AddMenuEntry(FCurveEditorCommands::Get().SetPostInfinityExtrapOscillate);
+							SubMenu.AddMenuEntry(FCurveEditorCommands::Get().SetPostInfinityExtrapLinear);
+							SubMenu.AddMenuEntry(FCurveEditorCommands::Get().SetPostInfinityExtrapConstant);
+						})
+					);
+
+					MenuBuilder.AddMenuSeparator();
+				}
 
 				// Filters
 				MenuBuilder.AddMenuEntry(FCurveEditorCommands::Get().OpenUserImplementableFilterWindow);
