Curve Editor: Fix set curve attributes operating all curves rather than just the selection

diff --git a/Engine/Source/Editor/CurveEditor/Private/SCurveEditorPanel.cpp b/Engine/Source/Editor/CurveEditor/Private/SCurveEditorPanel.cpp
--- a/Engine/Source/Editor/CurveEditor/Private/SCurveEditorPanel.cpp
+++ b/Engine/Source/Editor/CurveEditor/Private/SCurveEditorPanel.cpp
@@ -608,16 +608,16 @@ void SCurveEditorPanel::UpdateCommonCurveInfo()
 {
 	// Gather up common extended curve info for the current set of curves
 	TOptional<FCurveAttributes> AccumulatedCurveAttributes;
-	for (const FCurveModelID& CurveID : CurveEditor->GetEditedCurves())
+	for (const TTuple<FCurveModelID, FKeyHandleSet>& Pair : CurveEditor->Selection.GetAll())
 	{
 		FCurveAttributes Attributes;
 		
-		FCurveModel* Curve = CurveEditor->FindCurve(CurveID);
+		FCurveModel* Curve = CurveEditor->FindCurve(Pair.Key);
 		if (Curve)
 		{
 			Curve->GetCurveAttributes(Attributes);
 
-			// Some curves don't support extrapolation. We don't count them for determinine the accumulated state.
+			// Some curves don't support extrapolation. We don't count them for determine the accumulated state.
 			if (Attributes.HasPreExtrapolation() && Attributes.GetPreExtrapolation() == RCCE_None && Attributes.HasPostExtrapolation() && Attributes.GetPostExtrapolation() == RCCE_None)
 			{
 				continue;
@@ -777,10 +777,13 @@ void SCurveEditorPanel::SetCurveAttributes(FCurveAttributes CurveAttributes, FTe
 {
 	FScopedTransaction Transaction(Description);
 
-	for (const TTuple<FCurveModelID, TUniquePtr<FCurveModel>>& Pair : CurveEditor->GetCurves())
+	for (const TTuple<FCurveModelID, FKeyHandleSet>& Pair : CurveEditor->Selection.GetAll())
 	{
-		Pair.Value->Modify();
-		Pair.Value->SetCurveAttributes(CurveAttributes);
+		if (FCurveModel* Curve = CurveEditor->FindCurve(Pair.Key))
+		{
+			Curve->Modify();
+			Curve->SetCurveAttributes(CurveAttributes);
+		}
 	}
 }
 
