Curve Editor: Fix filter not opening for standalone curves since the tab manager may not exist.

diff --git a/Engine/Source/Editor/CurveEditor/Private/SCurveEditorPanel.cpp b/Engine/Source/Editor/CurveEditor/Private/SCurveEditorPanel.cpp
--- a/Engine/Source/Editor/CurveEditor/Private/SCurveEditorPanel.cpp
+++ b/Engine/Source/Editor/CurveEditor/Private/SCurveEditorPanel.cpp
@@ -981,10 +981,10 @@ FSlateIcon SCurveEditorPanel::GetCurveExtrapolationPostIcon() const
 void SCurveEditorPanel::ShowCurveFilterUI(TSubclassOf<UCurveEditorFilterBase> FilterClass)
 {
 	TSharedPtr<FTabManager> TabManager = WeakTabManager.Pin();
-	if (TabManager)
-	{
-		SCurveEditorFilterPanel::OpenDialog(TabManager.ToSharedRef(), CurveEditor.ToSharedRef(), FilterClass);
-	}
+	TSharedPtr<SDockTab> OwnerTab = TabManager.IsValid() ? TabManager->GetOwnerTab() : TSharedPtr<SDockTab>();
+	TSharedPtr<SWindow> RootWindow = OwnerTab.IsValid() ? OwnerTab->GetParentWindow() : TSharedPtr<SWindow>();
+
+	SCurveEditorFilterPanel::OpenDialog(RootWindow, CurveEditor.ToSharedRef(), FilterClass);
 }
 
 const FGeometry& SCurveEditorPanel::GetScrollPanelGeometry() const
diff --git a/Engine/Source/Editor/CurveEditor/Public/Filters/SCurveEditorFilterPanel.cpp b/Engine/Source/Editor/CurveEditor/Public/Filters/SCurveEditorFilterPanel.cpp
--- a/Engine/Source/Editor/CurveEditor/Public/Filters/SCurveEditorFilterPanel.cpp
+++ b/Engine/Source/Editor/CurveEditor/Public/Filters/SCurveEditorFilterPanel.cpp
@@ -165,7 +165,7 @@ bool SCurveEditorFilterPanel::CanApplyFilter() const
 }
 
 
-void SCurveEditorFilterPanel::OpenDialog(const TSharedRef<FTabManager>& TabManager, TSharedRef<FCurveEditor> InHostCurveEditor, TSubclassOf<UCurveEditorFilterBase> DefaultFilterClass)
+void SCurveEditorFilterPanel::OpenDialog(TSharedPtr<SWindow> RootWindow, TSharedRef<FCurveEditor> InHostCurveEditor, TSubclassOf<UCurveEditorFilterBase> DefaultFilterClass)
 {
 	TSharedPtr<SWindow> ExistingWindow = ExistingFilterWindow.Pin();
 	if (ExistingWindow.IsValid())
@@ -181,8 +181,6 @@ void SCurveEditorFilterPanel::OpenDialog(const TSharedRef<FTabManager>& TabManag
 			.SupportsMinimize(false)
 			.ClientSize(FVector2D(480, 360));
 
-		TSharedPtr<SDockTab> OwnerTab = TabManager->GetOwnerTab();
-		TSharedPtr<SWindow> RootWindow = OwnerTab.IsValid() ? OwnerTab->GetParentWindow() : TSharedPtr<SWindow>();
 		if (RootWindow.IsValid())
 		{
 			FSlateApplication::Get().AddWindowAsNativeChild(ExistingWindow.ToSharedRef(), RootWindow.ToSharedRef());
diff --git a/Engine/Source/Editor/CurveEditor/Public/Filters/SCurveEditorFilterPanel.h b/Engine/Source/Editor/CurveEditor/Public/Filters/SCurveEditorFilterPanel.h
--- a/Engine/Source/Editor/CurveEditor/Public/Filters/SCurveEditorFilterPanel.h
+++ b/Engine/Source/Editor/CurveEditor/Public/Filters/SCurveEditorFilterPanel.h
@@ -26,7 +26,7 @@ public:
 
 public:
 	/** Call this to request opening a window containing this panel. */
-	static void OpenDialog(const TSharedRef<FTabManager>& TabManager, TSharedRef<FCurveEditor> InHostCurveEditor, TSubclassOf<UCurveEditorFilterBase> DefaultFilterClass);
+	static void OpenDialog(TSharedPtr<SWindow> RootWindow, TSharedRef<FCurveEditor> InHostCurveEditor, TSubclassOf<UCurveEditorFilterBase> DefaultFilterClass);
 	
 	/** Closes the dialog if there is one open. */
 	static void CloseDialog();
