Sequencer: Fix CIS

#rnx

diff --git a/Engine/Source/Editor/CurveEditor/Private/SCurveEditorPanel.cpp b/Engine/Source/Editor/CurveEditor/Private/SCurveEditorPanel.cpp
--- a/Engine/Source/Editor/CurveEditor/Private/SCurveEditorPanel.cpp
+++ b/Engine/Source/Editor/CurveEditor/Private/SCurveEditorPanel.cpp
@@ -11,9 +11,12 @@
 #include "Framework/MultiBox/MultiBoxBuilder.h"
 #include "Widgets/SNullWidget.h"
 #include "Widgets/SOverlay.h"
+#include "Widgets/SWindow.h"
 #include "Widgets/Layout/SScrollBox.h"
 #include "Widgets/Layout/SSplitter.h"
 #include "Widgets/Text/STextBlock.h"
+#include "Widgets/Docking/SDockTab.h"
+#include "Framework/Docking/TabManager.h"
 #include "EditorStyleSet.h"
 #include "Misc/Attribute.h"
 #include "Algo/Sort.h"
