Fix curve editor having a white background

diff --git a/Engine/Source/Editor/CurveEditor/Private/SCurveEditorViewContainer.cpp b/Engine/Source/Editor/CurveEditor/Private/SCurveEditorViewContainer.cpp
--- a/Engine/Source/Editor/CurveEditor/Private/SCurveEditorViewContainer.cpp
+++ b/Engine/Source/Editor/CurveEditor/Private/SCurveEditorViewContainer.cpp
@@ -85,8 +85,9 @@ int32 SCurveEditorViewContainer::OnPaint(const FPaintArgs& Args, const FGeometry
 {
 	const ESlateDrawEffect DrawEffects = ShouldBeEnabled(bParentEnabled) ? ESlateDrawEffect::None : ESlateDrawEffect::DisabledEffect;
 	
-	FSlateDrawElement::MakeBox(OutDrawElements, LayerId, AllottedGeometry.ToPaintGeometry(),
-		FEditorStyle::GetBrush("ToolPanel.GroupBorder"), DrawEffects);
+	static const FName BackgroundBrushName("Brushes.Background");
+	const FSlateBrush* Background = FEditorStyle::GetBrush(BackgroundBrushName);
+	FSlateDrawElement::MakeBox(OutDrawElements, LayerId, AllottedGeometry.ToPaintGeometry(), Background, DrawEffects, Background->GetTint(InWidgetStyle));
 
 	SVerticalBox::OnPaint(Args, AllottedGeometry, MyCullingRect, OutDrawElements, LayerId, InWidgetStyle, bParentEnabled);
 
