Curve Editor: Fix crash in special case that adds an extra point oi constant interpolation segments. Need to increment past the index after inserting a point.

diff --git a/Engine/Source/Editor/CurveEditor/Private/SCurveEditorView.cpp b/Engine/Source/Editor/CurveEditor/Private/SCurveEditorView.cpp
--- a/Engine/Source/Editor/CurveEditor/Private/SCurveEditorView.cpp
+++ b/Engine/Source/Editor/CurveEditor/Private/SCurveEditorView.cpp
@@ -202,6 +202,7 @@ void SCurveEditorView::GetCurveDrawParams(TArray<FCurveDrawParams>& OutDrawParam
 				if (FVector2D::DotProduct(Line1, Line2) <= 0.5f)
 				{
 					Params.InterpolatingPoints.Insert(Params.InterpolatingPoints[Index] + Line2, Index + 1);
+					++Index;
 				}
 			}
 		}
@@ -304,4 +305,4 @@ void SCurveEditorView::GetCurveDrawParams(TArray<FCurveDrawParams>& OutDrawParam
 
 		OutDrawParams.Add(MoveTemp(Params));
 	}
-}
\ No newline at end of file
+}
