Curve Editor Selected Text

Fixes curve editor tree text to be inverted when selected

diff --git a/Engine/Source/Editor/CurveEditor/Private/Tree/SCurveEditorTree.cpp b/Engine/Source/Editor/CurveEditor/Private/Tree/SCurveEditorTree.cpp
--- a/Engine/Source/Editor/CurveEditor/Private/Tree/SCurveEditorTree.cpp
+++ b/Engine/Source/Editor/CurveEditor/Private/Tree/SCurveEditorTree.cpp
@@ -30,7 +30,7 @@ struct SCurveEditorTableRow : SMultiColumnTableRow<FCurveEditorTreeItemID>
 		TSharedPtr<FCurveEditor> CurveEditor = WeakCurveEditor.Pin();
 
 		const bool bIsMatch = CurveEditor.IsValid() && ( CurveEditor->GetTree()->GetFilterState(TreeItemID) == ECurveEditorTreeFilterState::Match );
-		return bIsMatch ? FSlateColor::UseForeground() : FSlateColor::UseSubduedForeground();
+		return bIsMatch ? GetForegroundBasedOnSelection() : FSlateColor::UseSubduedForeground();
 	}
 
 	virtual TSharedRef<SWidget> GenerateWidgetForColumn(const FName& InColumnName) override
