Sequencer:  Add Hot Keys To Nudge Selected Keys Left and Right.

diff --git a/Engine/Source/Editor/CurveEditor/Private/CurveEditor.cpp b/Engine/Source/Editor/CurveEditor/Private/CurveEditor.cpp
--- a/Engine/Source/Editor/CurveEditor/Private/CurveEditor.cpp
+++ b/Engine/Source/Editor/CurveEditor/Private/CurveEditor.cpp
@@ -343,6 +343,9 @@ void FCurveEditor::BindCommands()
 	CommandList->MapAction(FCurveEditorCommands::Get().ToggleExpandCollapseNodes, FExecuteAction::CreateSP(this, &FCurveEditor::ToggleExpandCollapseNodes, false));
 	CommandList->MapAction(FCurveEditorCommands::Get().ToggleExpandCollapseNodesAndDescendants, FExecuteAction::CreateSP(this, &FCurveEditor::ToggleExpandCollapseNodes, true));
 
+	CommandList->MapAction(FCurveEditorCommands::Get().TranslateSelectedKeysLeft, FExecuteAction::CreateSP(this, &FCurveEditor::TranslateSelectedKeysLeft));
+	CommandList->MapAction(FCurveEditorCommands::Get().TranslateSelectedKeysRight, FExecuteAction::CreateSP(this, &FCurveEditor::TranslateSelectedKeysRight));
+
 	CommandList->MapAction(FCurveEditorCommands::Get().StepToNextKey, FExecuteAction::CreateSP(this, &FCurveEditor::StepToNextKey));
 	CommandList->MapAction(FCurveEditorCommands::Get().StepToPreviousKey, FExecuteAction::CreateSP(this, &FCurveEditor::StepToPreviousKey));
 	CommandList->MapAction(FCurveEditorCommands::Get().StepForward, FExecuteAction::CreateSP(this, &FCurveEditor::StepForward), EUIActionRepeatMode::RepeatEnabled);
@@ -620,6 +623,53 @@ void FCurveEditor::ZoomToFitInternal(EAxisList::Type Axes, const TMap<FCurveMode
 	}
 }
 
+void FCurveEditor::TranslateSelectedKeys(double SecondsToAdd)
+{
+	
+	if (Selection.Count() > 0)
+	{
+		for (const TTuple<FCurveModelID, FKeyHandleSet>& Pair : Selection.GetAll())
+		{
+			if (FCurveModel* Curve = FindCurve(Pair.Key))
+			{
+				int32 NumKeys = Pair.Value.Num();
+
+				if (NumKeys > 0)
+				{
+					TArrayView<const FKeyHandle> KeyHandles = Pair.Value.AsArray();
+					TArray<FKeyPosition> KeyPositions;
+					KeyPositions.SetNum(KeyHandles.Num());
+
+					Curve->GetKeyPositions(KeyHandles, KeyPositions);
+
+					for (int KeyIndex = 0; KeyIndex < KeyPositions.Num(); ++KeyIndex)
+					{
+						KeyPositions[KeyIndex].InputValue += SecondsToAdd;
+					}
+					Curve->SetKeyPositions(KeyHandles, KeyPositions);
+				}
+			}
+		}
+	}
+}
+
+void FCurveEditor::TranslateSelectedKeysLeft()
+{
+	FScopedTransaction Transaction(LOCTEXT("TranslateKeysLeft", "Translate Keys Left"));
+	FFrameRate FrameRate = WeakTimeSliderController.Pin()->GetDisplayRate();
+	double SecondsToAdd =  -FrameRate.AsInterval();
+	TranslateSelectedKeys(SecondsToAdd);
+}
+
+void FCurveEditor::TranslateSelectedKeysRight()
+{
+	FScopedTransaction Transaction(LOCTEXT("TranslateKeyRight", "Translate Keys Right"));
+	FFrameRate FrameRate = WeakTimeSliderController.Pin()->GetDisplayRate();
+	double SecondsToAdd = FrameRate.AsInterval();
+
+	TranslateSelectedKeys(SecondsToAdd);
+}
+
 void FCurveEditor::StepToNextKey()
 {
 	if (!WeakTimeSliderController.IsValid())
diff --git a/Engine/Source/Editor/CurveEditor/Private/CurveEditorCommands.cpp b/Engine/Source/Editor/CurveEditor/Private/CurveEditorCommands.cpp
--- a/Engine/Source/Editor/CurveEditor/Private/CurveEditorCommands.cpp
+++ b/Engine/Source/Editor/CurveEditor/Private/CurveEditorCommands.cpp
@@ -81,6 +81,10 @@ void FCurveEditorCommands::RegisterCommands()
 	UI_COMMAND(SetAxisSnappingHorizontal, "X Only", "Snap transform tool axis movement to X direction.", EUserInterfaceActionType::Button, FInputChord());
 	UI_COMMAND(SetAxisSnappingVertical, "Y Only", "Snap transform tool axis movement to Y direction.", EUserInterfaceActionType::Button, FInputChord());	
 	
+	//Key Movement
+	UI_COMMAND(TranslateSelectedKeysLeft, "Translate Selected Keys Left", "Translate selected keys one frame to the left", EUserInterfaceActionType::Button, FInputChord(EModifierKey::Control, EKeys::Left));
+	UI_COMMAND(TranslateSelectedKeysRight, "Translate Selected Keys Right", "Translate selected keys one frame to the right", EUserInterfaceActionType::Button, FInputChord(EModifierKey::Control, EKeys::Right));
+
 	// Time Management
 	UI_COMMAND(StepToNextKey, "Step to Next Key", "Step to the next key", EUserInterfaceActionType::Button, FInputChord(EKeys::Period));
 	UI_COMMAND(StepToPreviousKey, "Step to Previous Key", "Step to the previous key", EUserInterfaceActionType::Button, FInputChord(EKeys::Comma));
diff --git a/Engine/Source/Editor/CurveEditor/Public/CurveEditor.h b/Engine/Source/Editor/CurveEditor/Public/CurveEditor.h
--- a/Engine/Source/Editor/CurveEditor/Public/CurveEditor.h
+++ b/Engine/Source/Editor/CurveEditor/Public/CurveEditor.h
@@ -174,6 +174,11 @@ public:
 	ICurveEditorToolExtension* GetCurrentTool() const;
 
 	FCurveEditorToolID AddTool(TUniquePtr<ICurveEditorToolExtension>&& InTool);
+	
+	/** Nudge left or right*/
+	void TranslateSelectedKeys(double SecondsToAdd);
+	void TranslateSelectedKeysLeft();
+	void TranslateSelectedKeysRight();
 
 	/** Step to next or previous key from the current time */
 	void StepToNextKey();
diff --git a/Engine/Source/Editor/CurveEditor/Public/CurveEditorCommands.h b/Engine/Source/Editor/CurveEditor/Public/CurveEditorCommands.h
--- a/Engine/Source/Editor/CurveEditor/Public/CurveEditorCommands.h
+++ b/Engine/Source/Editor/CurveEditor/Public/CurveEditorCommands.h
@@ -89,6 +89,10 @@ public:
 	TSharedPtr<FUICommandInfo> SetAxisSnappingHorizontal;
 	TSharedPtr<FUICommandInfo> SetAxisSnappingVertical;
 
+	//Key Movement
+	TSharedPtr<FUICommandInfo> TranslateSelectedKeysLeft;
+	TSharedPtr<FUICommandInfo> TranslateSelectedKeysRight;
+
 	// Time Management
 	TSharedPtr<FUICommandInfo> StepToNextKey;
 	TSharedPtr<FUICommandInfo> StepToPreviousKey;
