Curve Editor: Apply same jump to end logic as Sequencer

diff --git a/Engine/Source/Editor/CurveEditor/Private/CurveEditor.cpp b/Engine/Source/Editor/CurveEditor/Private/CurveEditor.cpp
--- a/Engine/Source/Editor/CurveEditor/Private/CurveEditor.cpp
+++ b/Engine/Source/Editor/CurveEditor/Private/CurveEditor.cpp
@@ -834,7 +834,17 @@ void FCurveEditor::JumpToEnd()
 		return;
 	}
 
-	WeakTimeSliderController.Pin()->SetScrubPosition(WeakTimeSliderController.Pin()->GetPlayRange().GetUpperBoundValue());
+	const bool bInsetDisplayFrame = IsInputSnappingEnabled();
+
+	FFrameRate TickResolution = WeakTimeSliderController.Pin()->GetTickResolution();
+	FFrameRate DisplayRate = WeakTimeSliderController.Pin()->GetDisplayRate();
+
+	// Calculate an offset from the end to go to. If they have snapping on (and the scrub style is a block) the last valid frame is represented as one
+	// whole display rate frame before the end, otherwise we just subtract a single frame which matches the behavior of hitting play and letting it run to the end.
+	FFrameTime OneFrame = bInsetDisplayFrame ? FFrameRate::TransformTime(FFrameTime(1), DisplayRate, TickResolution) : FFrameTime(1);
+	FFrameTime NewTime = WeakTimeSliderController.Pin()->GetPlayRange().GetUpperBoundValue() - OneFrame;
+
+	WeakTimeSliderController.Pin()->SetScrubPosition(NewTime);
 }
 
 void FCurveEditor::SetSelectionRangeStart()
