Curve Editor: Fix Pre-Extrap -> Pre-Infinity and Post-Extrap -> Post-Infinity

diff --git a/Engine/Source/Editor/CurveEditor/Private/CurveEditorContextMenu.cpp b/Engine/Source/Editor/CurveEditor/Private/CurveEditorContextMenu.cpp
--- a/Engine/Source/Editor/CurveEditor/Private/CurveEditorContextMenu.cpp
+++ b/Engine/Source/Editor/CurveEditor/Private/CurveEditorContextMenu.cpp
@@ -99,7 +99,7 @@ void FCurveEditorContextMenu::BuildMenu(FMenuBuilder& MenuBuilder, TSharedRef<FC
 
 					MenuBuilder.AddMenuSeparator();
 					
-					MenuBuilder.AddSubMenu(LOCTEXT("PreExtrapText", "Pre-Extrap"), FText(), FNewMenuDelegate::CreateLambda(
+					MenuBuilder.AddSubMenu(LOCTEXT("PreInfinityText", "Pre-Infinity"), FText(), FNewMenuDelegate::CreateLambda(
 						[](FMenuBuilder& SubMenu)
 						{
 							SubMenu.AddMenuEntry(FCurveEditorCommands::Get().SetPreInfinityExtrapCycle);
@@ -110,7 +110,7 @@ void FCurveEditorContextMenu::BuildMenu(FMenuBuilder& MenuBuilder, TSharedRef<FC
 						})
 					);
 
-					MenuBuilder.AddSubMenu(LOCTEXT("PostExtrapText", "Post-Extrap"), FText(), FNewMenuDelegate::CreateLambda(
+					MenuBuilder.AddSubMenu(LOCTEXT("PostInfinityText", "Post-Infinity"), FText(), FNewMenuDelegate::CreateLambda(
 						[](FMenuBuilder& SubMenu)
 						{
 							SubMenu.AddMenuEntry(FCurveEditorCommands::Get().SetPostInfinityExtrapCycle);
