Curve Editor: Increase the splitter handle width to make it easier to grab

diff --git a/Engine/Source/Editor/CurveEditor/Private/SCurveEditorPanel.cpp b/Engine/Source/Editor/CurveEditor/Private/SCurveEditorPanel.cpp
--- a/Engine/Source/Editor/CurveEditor/Private/SCurveEditorPanel.cpp
+++ b/Engine/Source/Editor/CurveEditor/Private/SCurveEditorPanel.cpp
@@ -233,7 +233,7 @@ void SCurveEditorPanel::Construct(const FArguments& InArgs, TSharedRef<FCurveEdi
 			SNew(SSplitter)
 			.Orientation(Orient_Horizontal)
 			.Style(FEditorStyle::Get(), "SplitterDark")
-			.PhysicalSplitterHandleSize(2.0f)
+			.PhysicalSplitterHandleSize(3.0f)
 
 			+ SSplitter::Slot()
 			.Value(InArgs._TreeSplitterWidth)
