Curve Editor: Frame on selected and neighboring key handles

diff --git a/Engine/Source/Editor/CurveEditor/Private/CurveEditor.cpp b/Engine/Source/Editor/CurveEditor/Private/CurveEditor.cpp
--- a/Engine/Source/Editor/CurveEditor/Private/CurveEditor.cpp
+++ b/Engine/Source/Editor/CurveEditor/Private/CurveEditor.cpp
@@ -521,9 +521,27 @@ void FCurveEditor::ZoomToFitInternal(EAxisList::Type Axes, const TMap<FCurveMode
 		}
 		else
 		{
-			// Zoom to the min/max of the specified key set
-			KeyPositionsScratch.SetNum(NumKeys, false);
-			Curve->GetKeyPositions(Pair.Value.AsArray(), KeyPositionsScratch);
+			// Zoom to the min/max of the specified key set and the neighbors
+			TArray<FKeyHandle> SelectedAndNeighborKeyHandles;
+			for (const FKeyHandle& SelectedKeyHandle : Pair.Value.AsArray())
+			{
+				TOptional<FKeyHandle> PreviousKeyHandle;
+				TOptional<FKeyHandle> NextKeyHandle;
+				Curve->GetNeighboringKeys(SelectedKeyHandle, PreviousKeyHandle, NextKeyHandle);
+
+				SelectedAndNeighborKeyHandles.Add(SelectedKeyHandle);
+				if (PreviousKeyHandle.IsSet())
+				{
+					SelectedAndNeighborKeyHandles.Add(PreviousKeyHandle.GetValue());
+				}
+				if (NextKeyHandle.IsSet())
+				{
+					SelectedAndNeighborKeyHandles.Add(NextKeyHandle.GetValue());
+				}
+			}
+
+			KeyPositionsScratch.SetNum(SelectedAndNeighborKeyHandles.Num(), false);
+			Curve->GetKeyPositions(SelectedAndNeighborKeyHandles, KeyPositionsScratch);
 			for (const FKeyPosition& Key : KeyPositionsScratch)
 			{
 				InputMin  = FMath::Min(InputMin, Key.InputValue);
