Curve Editor: Fix crash undoing inserting keys.

The panel's edit box's CurveIDToKeyProxies becomes stale on undo. Invalidate the curve model ID and relevant key handles on undo

diff --git a/Engine/Source/Editor/CurveEditor/Private/CurveEditor.cpp b/Engine/Source/Editor/CurveEditor/Private/CurveEditor.cpp
--- a/Engine/Source/Editor/CurveEditor/Private/CurveEditor.cpp
+++ b/Engine/Source/Editor/CurveEditor/Private/CurveEditor.cpp
@@ -1611,6 +1611,10 @@ void FCurveEditor::PostUndo(bool bSuccess)
 		if (!CurveModel)
 		{
 			Selection.Remove(Set.Key);
+			if (WeakPanel.IsValid())
+			{
+				WeakPanel.Pin()->InvalidateCurve(Set.Key);
+			}
 			continue;
 		}
 		// Get all of the key handles from this curve.
@@ -1627,6 +1631,10 @@ void FCurveEditor::PostUndo(bool bSuccess)
 			if (!KeyHandles.Contains(Handle))
 			{
 				Selection.Remove(Set.Key, ECurvePointType::Key, Handle);
+				if (WeakPanel.IsValid())
+				{
+					WeakPanel.Pin()->InvalidateCurve(Set.Key, Handle);
+				}
 			}
 		}
 	}
diff --git a/Engine/Source/Editor/CurveEditor/Private/SCurveEditorPanel.cpp b/Engine/Source/Editor/CurveEditor/Private/SCurveEditorPanel.cpp
--- a/Engine/Source/Editor/CurveEditor/Private/SCurveEditorPanel.cpp
+++ b/Engine/Source/Editor/CurveEditor/Private/SCurveEditorPanel.cpp
@@ -446,6 +446,26 @@ void SCurveEditorPanel::RemoveCurveFromViews(FCurveModelID InCurveID)
 	}
 }
 
+void SCurveEditorPanel::InvalidateCurve(FCurveModelID InvalidCurveModelID, FKeyHandle InvalidKeyHandle)
+{
+	// If key handle is invalidate, remove all 
+	if (InvalidKeyHandle == FKeyHandle::Invalid())
+	{
+		EditObjects->CurveIDToKeyProxies.Remove(InvalidCurveModelID);
+	}
+	else
+	{
+		for (TTuple<FCurveModelID, TMap<FKeyHandle, UObject*>>& OuterPair : EditObjects->CurveIDToKeyProxies)
+		{
+			if (OuterPair.Key == InvalidCurveModelID)
+			{
+				OuterPair.Value.Remove(InvalidKeyHandle);
+				break;
+			}
+		}
+	}
+}
+
 void SCurveEditorPanel::AddView(TSharedRef<SCurveEditorView> ViewToAdd)
 {
 	ExternalViews.Add(ViewToAdd);
diff --git a/Engine/Source/Editor/CurveEditor/Public/SCurveEditorPanel.h b/Engine/Source/Editor/CurveEditor/Public/SCurveEditorPanel.h
--- a/Engine/Source/Editor/CurveEditor/Public/SCurveEditorPanel.h
+++ b/Engine/Source/Editor/CurveEditor/Public/SCurveEditorPanel.h
@@ -125,6 +125,9 @@ class CURVEEDITOR_API SCurveEditorPanel : public SCompoundWidget
 	/** Get the last set View Mode for this UI. Utility function for the UI. */
 	ECurveEditorViewID GetViewMode() const { return DefaultViewID; }
 
+	/** Update based on invalidated curve model and key handle */
+	void InvalidateCurve(FCurveModelID InvalidCurveModelID, FKeyHandle InvalidKeyHandle = FKeyHandle::Invalid());
+
 private:
 	// SWidget Interface
 	virtual void Tick(const FGeometry& AllottedGeometry, const double InCurrentTime, const float InDeltaTime) override;
