Static Warning Fix for yet another overly sensitive static warning.
#trivial

diff --git a/Engine/Source/Editor/CurveEditor/Private/CurveEditorSettings.cpp b/Engine/Source/Editor/CurveEditor/Private/CurveEditorSettings.cpp
--- a/Engine/Source/Editor/CurveEditor/Private/CurveEditorSettings.cpp
+++ b/Engine/Source/Editor/CurveEditor/Private/CurveEditorSettings.cpp
@@ -137,6 +137,8 @@ FLinearColor UCurveEditorSettings::GetNextRandomColor()
 		IndexedColor.Add(FLinearColor(FColor::Blue));
 	}
 	FLinearColor Color = IndexedColor[NextIndex];
-	NextIndex = (++NextIndex) % IndexedColor.Num();
+	++NextIndex;
+	int32 NewIndex = (NextIndex % IndexedColor.Num());
+	NextIndex = NewIndex;
 	return Color;
 }
\ No newline at end of file
