Curve Editor: When adding a key to curve editors. If shift is down, don't disturb the original curve's tangent. (Remember to uncheck value snapping)

diff --git a/Engine/Source/Editor/CurveEditor/Private/Views/SInteractiveCurveEditorView.cpp b/Engine/Source/Editor/CurveEditor/Private/Views/SInteractiveCurveEditorView.cpp
--- a/Engine/Source/Editor/CurveEditor/Private/Views/SInteractiveCurveEditorView.cpp
+++ b/Engine/Source/Editor/CurveEditor/Private/Views/SInteractiveCurveEditorView.cpp
@@ -943,11 +943,33 @@ FReply SInteractiveCurveEditorView::OnMouseButtonDown(const FGeometry& MyGeometr
 					FScopedTransaction Transaction(LOCTEXT("InsertKey", "Insert Key"));
 
 					FCurveEditorScreenSpace CurveSpace = GetCurveSpace(HoveredCurve.GetValue());
-					FKeyAttributes DefaultAttributes = CurveEditor->GetDefaultKeyAttributes().Get();
-
+					FKeyAttributes KeyAttributes = CurveEditor->GetDefaultKeyAttributes().Get();
 					double MouseTime = CurveSpace.ScreenToSeconds(MousePixel.X);
 					double MouseValue = CurveSpace.ScreenToValue(MousePixel.Y);
 
+					// If shift is pressed. Keep the curve unchanged
+					if (MouseEvent.IsShiftDown())
+					{
+						KeyAttributes.SetTangentMode(RCTM_User);
+
+						const double DeltaTime = 0.1;
+						CurveToAddTo->Evaluate(MouseTime, MouseValue);
+
+						// Right
+						double NextTime = MouseTime + DeltaTime;
+						double NextValue = 0.0;
+						CurveToAddTo->Evaluate(NextTime, NextValue);
+						double RightTangent = (NextValue - MouseValue) / DeltaTime;
+						KeyAttributes.SetLeaveTangent(RightTangent);
+
+						// Left
+						double PrevTime = MouseTime - DeltaTime;
+						double PrevValue = 0.0;
+						CurveToAddTo->Evaluate(PrevTime, PrevValue);
+						double LeftTangent = (PrevValue - MouseValue) / DeltaTime;
+						KeyAttributes.SetArriveTangent(LeftTangent);
+					}
+
 					FCurveSnapMetrics SnapMetrics = CurveEditor->GetCurveSnapMetrics(HoveredCurve.GetValue());
 					MouseTime = SnapMetrics.SnapInputSeconds(MouseTime);
 					MouseValue = SnapMetrics.SnapOutput(MouseValue);
@@ -964,7 +986,7 @@ FReply SInteractiveCurveEditorView::OnMouseButtonDown(const FGeometry& MyGeometr
 					CurveToAddTo->Modify();
 
 					// Add a key on this curve
-					TOptional<FKeyHandle> NewKey = CurveToAddTo->AddKey(FKeyPosition(MouseTime, MouseValue), DefaultAttributes);
+					TOptional<FKeyHandle> NewKey = CurveToAddTo->AddKey(FKeyPosition(MouseTime, MouseValue), KeyAttributes);
 					if (NewKey.IsSet())
 					{
 						NewPoint = FCurvePointHandle(HoveredCurve.GetValue(), ECurvePointType::Key, NewKey.GetValue());
