CurveEditor: Further implemented the feature "Don't disturb the original curve when adding a key if shift key is down".

diff --git a/Engine/Source/Editor/CurveEditor/Private/Views/SInteractiveCurveEditorView.cpp b/Engine/Source/Editor/CurveEditor/Private/Views/SInteractiveCurveEditorView.cpp
--- a/Engine/Source/Editor/CurveEditor/Private/Views/SInteractiveCurveEditorView.cpp
+++ b/Engine/Source/Editor/CurveEditor/Private/Views/SInteractiveCurveEditorView.cpp
@@ -781,6 +781,34 @@ FText SInteractiveCurveEditorView::GetToolTipValueText() const
 	return CachedToolTipData.IsSet() ? CachedToolTipData->EvaluatedValue : FText();
 }
 
+double SInteractiveCurveEditorView::GetTangentValue(const double InTime, const double InValue, FCurveModel* CurveToAddTo, double DeltaTime) const
+{
+	// Data
+	double TargetTime = InTime + DeltaTime;				// The time to get tangent value. Could be left or right depending on is DeltaTime is negative or positive
+	double TargetValue = 0.0;							// The helper value to get Tangent value
+	CurveToAddTo->Evaluate(TargetTime, TargetValue);	// Initialize TargetValue by TargetTime
+	double TangentValue = (TargetValue - InValue) / FMath::Abs(DeltaTime);	// The tangent value to return
+	double PrevTangent = DBL_MAX;						// Used for determine whether the tangent is close to the limit
+	int32 Count = 10;									// Preventing we stuck in this function for too long
+
+	// Logic
+	// While the tangents not close enough and we haven't reach the max iteration time
+	while (!FMath::IsNearlyEqual(FMath::Abs(TangentValue), FMath::Abs(PrevTangent)) && Count > 0)
+	{
+		// Update previous tangent value and make delta time smaller
+		PrevTangent = TangentValue;
+		DeltaTime /= 2.0;
+		TargetTime = InTime + DeltaTime;
+
+		// Calculate a more precise tangent value
+		CurveToAddTo->Evaluate(TargetTime, TargetValue);
+		TangentValue = (TargetValue - InValue) / FMath::Abs(DeltaTime);
+
+		--Count;
+	}
+	return TangentValue;
+}
+
 void SInteractiveCurveEditorView::HandleDirectKeySelectionByMouse(TSharedPtr<FCurveEditor> CurveEditor, const FPointerEvent& MouseEvent, TOptional<FCurvePointHandle> MouseDownPoint)
 {
 	if(!MouseDownPoint.IsSet())
@@ -952,21 +980,18 @@ FReply SInteractiveCurveEditorView::OnMouseButtonDown(const FGeometry& MyGeometr
 					{
 						KeyAttributes.SetTangentMode(RCTM_User);
 
-						const double DeltaTime = 0.1;
+						// Estimated delta time to compute right and left tangents
+						double DeltaTime = 0.1;
+
+						// Make mouse value more accurate 
 						CurveToAddTo->Evaluate(MouseTime, MouseValue);
 
-						// Right
-						double NextTime = MouseTime + DeltaTime;
-						double NextValue = 0.0;
-						CurveToAddTo->Evaluate(NextTime, NextValue);
-						double RightTangent = (NextValue - MouseValue) / DeltaTime;
+						// Compute right tangent
+						double RightTangent = GetTangentValue(MouseTime, MouseValue, CurveToAddTo, DeltaTime);
 						KeyAttributes.SetLeaveTangent(RightTangent);
 
 						// Left
-						double PrevTime = MouseTime - DeltaTime;
-						double PrevValue = 0.0;
-						CurveToAddTo->Evaluate(PrevTime, PrevValue);
-						double LeftTangent = (PrevValue - MouseValue) / DeltaTime;
+						double LeftTangent = GetTangentValue(MouseTime, MouseValue, CurveToAddTo, -DeltaTime);
 						KeyAttributes.SetArriveTangent(LeftTangent);
 					}
 
diff --git a/Engine/Source/Editor/CurveEditor/Public/Views/SInteractiveCurveEditorView.h b/Engine/Source/Editor/CurveEditor/Public/Views/SInteractiveCurveEditorView.h
--- a/Engine/Source/Editor/CurveEditor/Public/Views/SInteractiveCurveEditorView.h
+++ b/Engine/Source/Editor/CurveEditor/Public/Views/SInteractiveCurveEditorView.h
@@ -15,6 +15,7 @@ struct FCurveEditorScreenSpace;
 struct FOptionalSize;
 struct FCurveEditorDelayedDrag;
 class IMenu;
+class FCurveModel;
 
 namespace CurveViewConstants
 {
@@ -123,6 +124,17 @@ private:
 	FText GetToolTipTimeText() const;
 	FText GetToolTipValueText() const;
 
+	/**
+	 * Returns the proper tangent value so we can keep the curve remain the original shape
+	 *
+	 * @param InTime		The time we are trying to add a key to
+	 * @param InValue		The value we are trying to add a key to
+	 * @param CurveToAddTo  The curve we are trying to add a key to
+	 * @param DeltaTime		Negative to get the left tangent, positive for right. Remember to use FMath::Abs() when needed
+	 * @return				The tangent value relatives to the DeltaTime upon mouse click's position
+	 */
+	double GetTangentValue(const double InTime, const double InValue, FCurveModel* CurveToAddTo, double DeltaTime) const;
+
 	/*~ Command binding callbacks */
 	void AddKeyAtScrubTime(TSet<FCurveModelID> ForCurves);
 	void AddKeyAtMousePosition(TSet<FCurveModelID> ForCurves);
