Fix to include MenuBuilder

diff --git a/Engine/Source/Editor/CurveEditor/Public/Views/SInteractiveCurveEditorView.h b/Engine/Source/Editor/CurveEditor/Public/Views/SInteractiveCurveEditorView.h
--- a/Engine/Source/Editor/CurveEditor/Public/Views/SInteractiveCurveEditorView.h
+++ b/Engine/Source/Editor/CurveEditor/Public/Views/SInteractiveCurveEditorView.h
@@ -4,6 +4,7 @@
 
 #include "Widgets/DeclarativeSyntaxSupport.h"
 #include "Widgets/SCompoundWidget.h"
+#include "Framework/MultiBox/MultiBoxBuilder.h"
 #include "CurveEditorTypes.h"
 #include "SCurveEditorView.h"
 #include "ICurveEditorDragOperation.h"
