Curve Editor Stacked View - fixes label rendering

diff --git a/Engine/Source/Editor/CurveEditor/Private/Views/SCurveEditorViewStacked.cpp b/Engine/Source/Editor/CurveEditor/Private/Views/SCurveEditorViewStacked.cpp
--- a/Engine/Source/Editor/CurveEditor/Private/Views/SCurveEditorViewStacked.cpp
+++ b/Engine/Source/Editor/CurveEditor/Private/Views/SCurveEditorViewStacked.cpp
@@ -222,7 +222,7 @@ void SCurveEditorViewStacked::DrawLabels(const FGeometry& AllottedGeometry, cons
 	const double ValuePerPixel = 1.0 / StackedHeight;
 	const double ValueSpacePadding = StackedPadding * ValuePerPixel;
 
-	const FSlateFontInfo FontInfo = FCoreStyle::Get().GetFontStyle("FontAwesome.11");
+	const FSlateFontInfo FontInfo = FAppStyle::Get().GetFontStyle("NormalFont");
 	const FVector2D LocalSize = AllottedGeometry.GetLocalSize();
 	const FCurveEditorScreenSpaceV ViewSpace = GetViewSpace();
 
