Animseq curves show as a flat line
#fix apply Min() operation with fractional representation instead of percentage

diff --git a/Engine/Source/Editor/CurveEditor/Private/CurveEditor.cpp b/Engine/Source/Editor/CurveEditor/Private/CurveEditor.cpp
--- a/Engine/Source/Editor/CurveEditor/Private/CurveEditor.cpp
+++ b/Engine/Source/Editor/CurveEditor/Private/CurveEditor.cpp
@@ -633,7 +633,7 @@ void FCurveEditor::ZoomToFitInternal(EAxisList::Type Axes, const TMap<FCurveMode
 				PanelWidth = View->GetViewSpace().GetPhysicalWidth();
 			}
 			
-			double InputPercentage = PanelWidth != 0 ? FMath::Min(Settings->GetFrameInputPadding() / (float)PanelWidth, 50.0) : 0.1; // Cannot pad more than half the width
+			double InputPercentage = PanelWidth != 0 ? FMath::Min(Settings->GetFrameInputPadding() / (float)PanelWidth, 0.5) : 0.1; // Cannot pad more than half the width
 
 			const double MinInputZoom = InputSnapEnabledAttribute.Get() ? InputSnapRateAttribute.Get().AsInterval() : 0.00001;
 			const double InputPadding = FMath::Max((InputMax - InputMin) * InputPercentage, MinInputZoom);
@@ -674,7 +674,7 @@ void FCurveEditor::ZoomToFitInternal(EAxisList::Type Axes, const TMap<FCurveMode
 				PanelHeight = View->GetViewSpace().GetPhysicalHeight();
 			}
 
-			double OutputPercentage = PanelHeight != 0 ? FMath::Min(Settings->GetFrameOutputPadding() / (float)PanelHeight, 50.0) : 0.1; // Cannot pad more than half the height
+			double OutputPercentage = PanelHeight != 0 ? FMath::Min(Settings->GetFrameOutputPadding() / (float)PanelHeight, 0.5) : 0.1; // Cannot pad more than half the height
 
 			constexpr double MinOutputZoom = 0.00001;
 			const double OutputPadding = FMath::Max((OutputMax - OutputMin) * OutputPercentage, MinOutputZoom);
