Curve Editor: Refactor to only pin once

diff --git a/Engine/Source/Editor/CurveEditor/Private/CurveEditor.cpp b/Engine/Source/Editor/CurveEditor/Private/CurveEditor.cpp
--- a/Engine/Source/Editor/CurveEditor/Private/CurveEditor.cpp
+++ b/Engine/Source/Editor/CurveEditor/Private/CurveEditor.cpp
@@ -689,7 +689,6 @@ void FCurveEditor::ZoomToFitInternal(EAxisList::Type Axes, const TMap<FCurveMode
 
 void FCurveEditor::TranslateSelectedKeys(double SecondsToAdd)
 {
-	
 	if (Selection.Count() > 0)
 	{
 		for (const TTuple<FCurveModelID, FKeyHandleSet>& Pair : Selection.GetAll())
@@ -719,16 +718,28 @@ void FCurveEditor::TranslateSelectedKeys(double SecondsToAdd)
 
 void FCurveEditor::TranslateSelectedKeysLeft()
 {
+	TSharedPtr<ITimeSliderController> TimeSliderController = WeakTimeSliderController.Pin();
+	if (!TimeSliderController.IsValid())
+	{
+		return;
+	}
+
 	FScopedTransaction Transaction(LOCTEXT("TranslateKeysLeft", "Translate Keys Left"));
-	FFrameRate FrameRate = WeakTimeSliderController.Pin()->GetDisplayRate();
+	FFrameRate FrameRate = TimeSliderController->GetDisplayRate();
 	double SecondsToAdd =  -FrameRate.AsInterval();
 	TranslateSelectedKeys(SecondsToAdd);
 }
 
 void FCurveEditor::TranslateSelectedKeysRight()
 {
+	TSharedPtr<ITimeSliderController> TimeSliderController = WeakTimeSliderController.Pin();
+	if (!TimeSliderController.IsValid())
+	{
+		return;
+	}
+
 	FScopedTransaction Transaction(LOCTEXT("TranslateKeyRight", "Translate Keys Right"));
-	FFrameRate FrameRate = WeakTimeSliderController.Pin()->GetDisplayRate();
+	FFrameRate FrameRate = TimeSliderController->GetDisplayRate();
 	double SecondsToAdd = FrameRate.AsInterval();
 
 	TranslateSelectedKeys(SecondsToAdd);
@@ -736,14 +747,15 @@ void FCurveEditor::TranslateSelectedKeysRight()
 
 void FCurveEditor::StepToNextKey()
 {
-	if (!WeakTimeSliderController.IsValid())
+	TSharedPtr<ITimeSliderController> TimeSliderController = WeakTimeSliderController.Pin();
+	if (!TimeSliderController.IsValid())
 	{
 		return;
 	}
 
-	FFrameRate TickResolution = WeakTimeSliderController.Pin()->GetTickResolution();
+	FFrameRate TickResolution = TimeSliderController->GetTickResolution();
 
-	double CurrentTime = TickResolution.AsSeconds(WeakTimeSliderController.Pin()->GetScrubPosition());
+	double CurrentTime = TickResolution.AsSeconds(TimeSliderController->GetScrubPosition());
 
 	TOptional<double> NextTime;
 	TOptional<double> MinTime;
@@ -762,7 +774,7 @@ void FCurveEditor::StepToNextKey()
 			KeyPositions.SetNum(KeyHandles.Num());
 			CurveModel->GetKeyPositions(TArrayView<FKeyHandle>(KeyHandles), KeyPositions);
 
-			for (FKeyPosition KeyPosition : KeyPositions)
+			for (const FKeyPosition& KeyPosition : KeyPositions)
 			{
 				if (KeyPosition.InputValue > CurrentTime)
 				{
@@ -784,24 +796,25 @@ void FCurveEditor::StepToNextKey()
 
 	if (NextTime.IsSet())
 	{
-		WeakTimeSliderController.Pin()->SetScrubPosition(NextTime.GetValue() * TickResolution,/*bEvaluate*/ true);
+		TimeSliderController->SetScrubPosition(NextTime.GetValue() * TickResolution,/*bEvaluate*/ true);
 	}
 	else if (MinTime.IsSet())
 	{
-		WeakTimeSliderController.Pin()->SetScrubPosition(MinTime.GetValue() * TickResolution, /*bEvaluate*/ true);
+		TimeSliderController->SetScrubPosition(MinTime.GetValue() * TickResolution, /*bEvaluate*/ true);
 	}
 }
 
 void FCurveEditor::StepToPreviousKey()
 {
-	if (!WeakTimeSliderController.IsValid())
+	TSharedPtr<ITimeSliderController> TimeSliderController = WeakTimeSliderController.Pin();
+	if (!TimeSliderController.IsValid())
 	{
 		return;
 	}
 
-	FFrameRate TickResolution = WeakTimeSliderController.Pin()->GetTickResolution();
+	FFrameRate TickResolution = TimeSliderController->GetTickResolution();
 
-	double CurrentTime = TickResolution.AsSeconds(WeakTimeSliderController.Pin()->GetScrubPosition());
+	double CurrentTime = TickResolution.AsSeconds(TimeSliderController->GetScrubPosition());
 
 	TOptional<double> PreviousTime;
 	TOptional<double> MaxTime;
@@ -820,7 +833,7 @@ void FCurveEditor::StepToPreviousKey()
 			KeyPositions.SetNum(KeyHandles.Num());
 			CurveModel->GetKeyPositions(TArrayView<FKeyHandle>(KeyHandles), KeyPositions);
 
-			for (FKeyPosition KeyPosition : KeyPositions)
+			for (const FKeyPosition& KeyPosition : KeyPositions)
 			{
 				if (KeyPosition.InputValue < CurrentTime)
 				{
@@ -842,121 +855,128 @@ void FCurveEditor::StepToPreviousKey()
 
 	if (PreviousTime.IsSet())
 	{
-		WeakTimeSliderController.Pin()->SetScrubPosition(PreviousTime.GetValue() * TickResolution,/*bEvaluate*/ true);
+		TimeSliderController->SetScrubPosition(PreviousTime.GetValue() * TickResolution,/*bEvaluate*/ true);
 	}
 	else if (MaxTime.IsSet())
 	{
-		WeakTimeSliderController.Pin()->SetScrubPosition(MaxTime.GetValue() * TickResolution, /*bEvaluate*/ true);
+		TimeSliderController->SetScrubPosition(MaxTime.GetValue() * TickResolution, /*bEvaluate*/ true);
 	}
 }
 
 
 void FCurveEditor::StepForward()
 {
-	if (!WeakTimeSliderController.IsValid())
+	TSharedPtr<ITimeSliderController> TimeSliderController = WeakTimeSliderController.Pin();
+	if (!TimeSliderController.IsValid())
 	{
 		return;
 	}
 
-	FFrameRate TickResolution = WeakTimeSliderController.Pin()->GetTickResolution();
-	FFrameRate DisplayRate = WeakTimeSliderController.Pin()->GetDisplayRate();
+	FFrameRate TickResolution = TimeSliderController->GetTickResolution();
+	FFrameRate DisplayRate = TimeSliderController->GetDisplayRate();
 
 	FFrameTime OneFrame = FFrameRate::TransformTime(FFrameTime(1), DisplayRate, TickResolution);
 
-	WeakTimeSliderController.Pin()->SetScrubPosition(WeakTimeSliderController.Pin()->GetScrubPosition() + OneFrame, /*bEvaluate*/ true);
+	TimeSliderController->SetScrubPosition(TimeSliderController->GetScrubPosition() + OneFrame, /*bEvaluate*/ true);
 }
 
 void FCurveEditor::StepBackward()
 {
-	if (!WeakTimeSliderController.IsValid())
+	TSharedPtr<ITimeSliderController> TimeSliderController = WeakTimeSliderController.Pin();
+	if (!TimeSliderController.IsValid())
 	{
 		return;
 	}
 
-	FFrameRate TickResolution = WeakTimeSliderController.Pin()->GetTickResolution();
-	FFrameRate DisplayRate = WeakTimeSliderController.Pin()->GetDisplayRate();
+	FFrameRate TickResolution = TimeSliderController->GetTickResolution();
+	FFrameRate DisplayRate = TimeSliderController->GetDisplayRate();
 
 	FFrameTime OneFrame = FFrameRate::TransformTime(FFrameTime(1), DisplayRate, TickResolution);
 
-	WeakTimeSliderController.Pin()->SetScrubPosition(WeakTimeSliderController.Pin()->GetScrubPosition() - OneFrame, /*bEvaluate*/ true);
+	TimeSliderController->SetScrubPosition(TimeSliderController->GetScrubPosition() - OneFrame, /*bEvaluate*/ true);
 }
 
 void FCurveEditor::JumpToStart()
 {
-	if (!WeakTimeSliderController.IsValid())
+	TSharedPtr<ITimeSliderController> TimeSliderController = WeakTimeSliderController.Pin();
+	if (!TimeSliderController.IsValid())
 	{
 		return;
 	}
 
-	WeakTimeSliderController.Pin()->SetScrubPosition(WeakTimeSliderController.Pin()->GetPlayRange().GetLowerBoundValue(), /*bEvaluate*/ true);
+	TimeSliderController->SetScrubPosition(TimeSliderController->GetPlayRange().GetLowerBoundValue(), /*bEvaluate*/ true);
 }
 
 void FCurveEditor::JumpToEnd()
 {
-	if (!WeakTimeSliderController.IsValid())
+	TSharedPtr<ITimeSliderController> TimeSliderController = WeakTimeSliderController.Pin();
+	if (!TimeSliderController.IsValid())
 	{
 		return;
 	}
 
 	const bool bInsetDisplayFrame = IsInputSnappingEnabled();
 
-	FFrameRate TickResolution = WeakTimeSliderController.Pin()->GetTickResolution();
-	FFrameRate DisplayRate = WeakTimeSliderController.Pin()->GetDisplayRate();
+	FFrameRate TickResolution = TimeSliderController->GetTickResolution();
+	FFrameRate DisplayRate = TimeSliderController->GetDisplayRate();
 
 	// Calculate an offset from the end to go to. If they have snapping on (and the scrub style is a block) the last valid frame is represented as one
 	// whole display rate frame before the end, otherwise we just subtract a single frame which matches the behavior of hitting play and letting it run to the end.
 	FFrameTime OneFrame = bInsetDisplayFrame ? FFrameRate::TransformTime(FFrameTime(1), DisplayRate, TickResolution) : FFrameTime(1);
-	FFrameTime NewTime = WeakTimeSliderController.Pin()->GetPlayRange().GetUpperBoundValue() - OneFrame;
+	FFrameTime NewTime = TimeSliderController->GetPlayRange().GetUpperBoundValue() - OneFrame;
 
-	WeakTimeSliderController.Pin()->SetScrubPosition(NewTime, /*bEvaluate*/ true);
+	TimeSliderController->SetScrubPosition(NewTime, /*bEvaluate*/ true);
 }
 
 void FCurveEditor::SetSelectionRangeStart()
 {
-	if (!WeakTimeSliderController.IsValid())
+	TSharedPtr<ITimeSliderController> TimeSliderController = WeakTimeSliderController.Pin();
+	if (!TimeSliderController.IsValid())
 	{
 		return;
 	}
 
-	FFrameNumber LocalTime = WeakTimeSliderController.Pin()->GetScrubPosition().FrameNumber;
-	FFrameNumber UpperBound = WeakTimeSliderController.Pin()->GetSelectionRange().GetUpperBoundValue();
+	FFrameNumber LocalTime = TimeSliderController->GetScrubPosition().FrameNumber;
+	FFrameNumber UpperBound = TimeSliderController->GetSelectionRange().GetUpperBoundValue();
 	if (UpperBound <= LocalTime)
 	{
-		WeakTimeSliderController.Pin()->SetSelectionRange(TRange<FFrameNumber>(LocalTime, LocalTime + 1));
+		TimeSliderController->SetSelectionRange(TRange<FFrameNumber>(LocalTime, LocalTime + 1));
 	}
 	else
 	{
-		WeakTimeSliderController.Pin()->SetSelectionRange(TRange<FFrameNumber>(LocalTime, UpperBound));
+		TimeSliderController->SetSelectionRange(TRange<FFrameNumber>(LocalTime, UpperBound));
 	}
 }
 
 void FCurveEditor::SetSelectionRangeEnd()
 {
-	if (!WeakTimeSliderController.IsValid())
+	TSharedPtr<ITimeSliderController> TimeSliderController = WeakTimeSliderController.Pin();
+	if (!TimeSliderController.IsValid())
 	{
 		return;
 	}
 
-	FFrameNumber LocalTime = WeakTimeSliderController.Pin()->GetScrubPosition().FrameNumber;
-	FFrameNumber LowerBound = WeakTimeSliderController.Pin()->GetSelectionRange().GetLowerBoundValue();
+	FFrameNumber LocalTime = TimeSliderController->GetScrubPosition().FrameNumber;
+	FFrameNumber LowerBound = TimeSliderController->GetSelectionRange().GetLowerBoundValue();
 	if (LowerBound >= LocalTime)
 	{
-		WeakTimeSliderController.Pin()->SetSelectionRange(TRange<FFrameNumber>(LocalTime - 1, LocalTime));
+		TimeSliderController->SetSelectionRange(TRange<FFrameNumber>(LocalTime - 1, LocalTime));
 	}
 	else
 	{
-		WeakTimeSliderController.Pin()->SetSelectionRange(TRange<FFrameNumber>(LowerBound, LocalTime));
+		TimeSliderController->SetSelectionRange(TRange<FFrameNumber>(LowerBound, LocalTime));
 	}
 }
 
 void FCurveEditor::ClearSelectionRange()
 {
-	if (!WeakTimeSliderController.IsValid())
+	TSharedPtr<ITimeSliderController> TimeSliderController = WeakTimeSliderController.Pin();
+	if (!TimeSliderController.IsValid())
 	{
 		return;
 	}
 
-	WeakTimeSliderController.Pin()->SetSelectionRange(TRange<FFrameNumber>::Empty());
+	TimeSliderController->SetSelectionRange(TRange<FFrameNumber>::Empty());
 }
 
 void FCurveEditor::SelectAllKeys()
@@ -1397,11 +1417,12 @@ void FCurveEditor::PasteKeys(TSet<FCurveModelID> CurveModelIDs)
 
 		if (bApplyOffset)
 		{
-			if (WeakTimeSliderController.IsValid())
+			TSharedPtr<ITimeSliderController> TimeSliderController = WeakTimeSliderController.Pin();
+			if (TimeSliderController.IsValid())
 			{
-				FFrameRate TickResolution = WeakTimeSliderController.Pin()->GetTickResolution();
+				FFrameRate TickResolution = TimeSliderController->GetTickResolution();
 
-				TimeOffset = WeakTimeSliderController.Pin()->GetScrubPosition() / TickResolution;
+				TimeOffset = TimeSliderController->GetScrubPosition() / TickResolution;
 			}
 			else
 			{
