Curve Editor: Fix copy/paste when keys are selected

Get the curve model IDs from the key selection, not the keys of the tree selection.

diff --git a/Engine/Source/Editor/CurveEditor/Private/CurveEditor.cpp b/Engine/Source/Editor/CurveEditor/Private/CurveEditor.cpp
--- a/Engine/Source/Editor/CurveEditor/Private/CurveEditor.cpp
+++ b/Engine/Source/Editor/CurveEditor/Private/CurveEditor.cpp
@@ -1345,26 +1345,29 @@ void FCurveEditor::PasteKeys(TSet<FCurveModelID> CurveModelIDs)
 			TArray<FCurveEditorTreeItemID> NodesToSearch;
 
 			// Try nodes with selected keys
-			GetTreeSelection().GetKeys(NodesToSearch);
+			for (const TTuple<FCurveModelID, FKeyHandleSet>& Pair : Selection.GetAll())
+			{
+				CurveModelIDs.Add(Pair.Key);
+			}
 
 			// Try selected nodes
-			if (NodesToSearch.Num() == 0)
+			if (CurveModelIDs.Num() == 0)
 			{
 				for (const TTuple<FCurveEditorTreeItemID, ECurveEditorTreeSelectionState>& Pair : GetTreeSelection())
 				{
 					NodesToSearch.Add(Pair.Key);
 				}
-			}
 
-			// If no curves are selected, paste to the entire tree using fully qualified long names
-			if (NodesToSearch.Num() == 0)
-			{
-				bSelectionNeedsLongNames = true;
-				bAllCopiedCurvesLongNameEqual = false;
-				if (Tree.GetAllItems().GetKeys(NodesToSearch) == 0)
+				// If no curves are selected, paste to the entire tree using fully qualified long names
+				if (NodesToSearch.Num() == 0)
 				{
-					// If we don't have any curves to paste in to, exit now
-					return;
+					bSelectionNeedsLongNames = true;
+					bAllCopiedCurvesLongNameEqual = false;
+					if (Tree.GetAllItems().GetKeys(NodesToSearch) == 0)
+					{
+						// If we don't have any curves to paste in to, exit now
+						return;
+					}
 				}
 			}
 
