Curve Editor: Add option to Snap Time to Selection

diff --git a/Engine/Source/Editor/CurveEditor/Private/CurveEditor.cpp b/Engine/Source/Editor/CurveEditor/Private/CurveEditor.cpp
--- a/Engine/Source/Editor/CurveEditor/Private/CurveEditor.cpp
+++ b/Engine/Source/Editor/CurveEditor/Private/CurveEditor.cpp
@@ -413,6 +413,12 @@ void FCurveEditor::BindCommands()
 		FIsActionChecked::CreateLambda( [CurveSettings]{ return CurveSettings->GetAutoFrameCurveEditor(); } )
 	);
 
+	CommandList->MapAction(FCurveEditorCommands::Get().ToggleSnapTimeToSelection,
+		FExecuteAction::CreateLambda( [CurveSettings]{ CurveSettings->SetSnapTimeToSelection( !CurveSettings->GetSnapTimeToSelection() ); } ),
+		FCanExecuteAction(),
+		FIsActionChecked::CreateLambda( [CurveSettings]{ return CurveSettings->GetSnapTimeToSelection(); } )
+	);
+
 	CommandList->MapAction(FCurveEditorCommands::Get().ToggleShowCurveEditorCurveToolTips,
 		FExecuteAction::CreateLambda( [CurveSettings]{ CurveSettings->SetShowCurveEditorCurveToolTips( !CurveSettings->GetShowCurveEditorCurveToolTips() ); } ),
 		FCanExecuteAction(),
@@ -745,6 +751,53 @@ void FCurveEditor::TranslateSelectedKeysRight()
 	TranslateSelectedKeys(SecondsToAdd);
 }
 
+void FCurveEditor::SnapToSelectedKey()
+{
+	TSharedPtr<ITimeSliderController> TimeSliderController = WeakTimeSliderController.Pin();
+	if (!TimeSliderController.IsValid())
+	{
+		return;
+	}
+
+	FFrameRate TickResolution = TimeSliderController->GetTickResolution();
+
+	TOptional<double> MinTime;
+
+	for (const TTuple<FCurveModelID, FKeyHandleSet>& Pair : Selection.GetAll())
+	{
+		if (FCurveModel* Curve = FindCurve(Pair.Key))
+		{
+			int32 NumKeys = Pair.Value.Num();
+
+			if (NumKeys > 0)
+			{
+				TArrayView<const FKeyHandle> KeyHandles = Pair.Value.AsArray();
+				TArray<FKeyPosition> KeyPositions;
+				KeyPositions.SetNum(KeyHandles.Num());
+
+				Curve->GetKeyPositions(KeyHandles, KeyPositions);
+
+				for (const FKeyPosition& KeyPosition : KeyPositions)
+				{
+					if (MinTime.IsSet())
+					{
+						MinTime = FMath::Min(KeyPosition.InputValue, MinTime.GetValue());
+					}
+					else
+					{
+						MinTime = KeyPosition.InputValue;
+					}
+				}
+			}
+		}
+	}
+
+	if (MinTime.IsSet())
+	{
+		TimeSliderController->SetScrubPosition(MinTime.GetValue() * TickResolution);		
+	}
+}
+
 void FCurveEditor::StepToNextKey()
 {
 	TSharedPtr<ITimeSliderController> TimeSliderController = WeakTimeSliderController.Pin();
diff --git a/Engine/Source/Editor/CurveEditor/Private/CurveEditorCommands.cpp b/Engine/Source/Editor/CurveEditor/Private/CurveEditorCommands.cpp
--- a/Engine/Source/Editor/CurveEditor/Private/CurveEditorCommands.cpp
+++ b/Engine/Source/Editor/CurveEditor/Private/CurveEditorCommands.cpp
@@ -50,6 +50,7 @@ void FCurveEditorCommands::RegisterCommands()
 	UI_COMMAND(SetNoTangentsVisibility, "No Tangents", "Show no tangents in the curve editor.", EUserInterfaceActionType::RadioButton, FInputChord());
 
 	UI_COMMAND(ToggleAutoFrameCurveEditor, "Auto Frame Curves", "Auto frame curves when they are selected.", EUserInterfaceActionType::ToggleButton, FInputChord() );
+	UI_COMMAND(ToggleSnapTimeToSelection, "Snap Time to Selection", "Snap the current time to the first selected key time.", EUserInterfaceActionType::ToggleButton, FInputChord() );
 	UI_COMMAND(ToggleShowCurveEditorCurveToolTips, "Curve Tool Tips", "Show a tool tip with name and values when hovering over a curve.", EUserInterfaceActionType::ToggleButton, FInputChord() );
 
 	UI_COMMAND(AddKeyHovered, "Add Key", "Add a new key to this curve at the current position.", EUserInterfaceActionType::Button, FInputChord(EKeys::MiddleMouseButton) );
diff --git a/Engine/Source/Editor/CurveEditor/Private/CurveEditorSettings.cpp b/Engine/Source/Editor/CurveEditor/Private/CurveEditorSettings.cpp
--- a/Engine/Source/Editor/CurveEditor/Private/CurveEditorSettings.cpp
+++ b/Engine/Source/Editor/CurveEditor/Private/CurveEditorSettings.cpp
@@ -10,6 +10,7 @@ UCurveEditorSettings::UCurveEditorSettings()
 	bShowCurveEditorCurveToolTips = true;
 	TangentVisibility = ECurveEditorTangentVisibility::SelectedKeys;
 	ZoomPosition = ECurveEditorZoomPosition::CurrentTime;
+	bSnapTimeToSelection = false;
 
 	ParentSpaceCustomColor = FLinearColor(.93, .31, .19); //pastel orange
 	WorldSpaceCustomColor = FLinearColor(.198, .610, .558); //pastel teal
@@ -99,6 +100,20 @@ void UCurveEditorSettings::SetZoomPosition(ECurveEditorZoomPosition InZoomPositi
 	}
 }
 
+bool UCurveEditorSettings::GetSnapTimeToSelection() const
+{
+	return bSnapTimeToSelection;
+}
+
+void UCurveEditorSettings::SetSnapTimeToSelection(bool bInSnapTimeToSelection)
+{
+	if (bSnapTimeToSelection != bInSnapTimeToSelection)
+	{
+		bSnapTimeToSelection = bInSnapTimeToSelection;
+		SaveConfig();
+	}
+}
+
 TOptional<FLinearColor> UCurveEditorSettings::GetCustomColor(UClass* InClass, const FString& InPropertyName) const
 {
 	TOptional<FLinearColor> Color;
diff --git a/Engine/Source/Editor/CurveEditor/Private/SCurveEditorPanel.cpp b/Engine/Source/Editor/CurveEditor/Private/SCurveEditorPanel.cpp
--- a/Engine/Source/Editor/CurveEditor/Private/SCurveEditorPanel.cpp
+++ b/Engine/Source/Editor/CurveEditor/Private/SCurveEditorPanel.cpp
@@ -434,6 +434,7 @@ void SCurveEditorPanel::Tick(const FGeometry& AllottedGeometry, const double InC
 
 	UpdateCommonCurveInfo();
 	UpdateEditBox();
+	UpdateTime();
 
 	CachedSelectionSerialNumber = CurveEditor->Selection.GetSerialNumber();
 }
@@ -688,6 +689,20 @@ void SCurveEditorPanel::OnCurveEditorToolChanged(FCurveEditorToolID InToolId)
 	ToolPropertiesPanel->OnToolChanged(InToolId);
 }
 
+void SCurveEditorPanel::UpdateTime()
+{
+	const FCurveEditorSelection& Selection = CurveEditor->Selection;
+	if (CachedSelectionSerialNumber == Selection.GetSerialNumber())
+	{
+		return;
+	}
+
+	if (CurveEditor->GetSettings()->GetSnapTimeToSelection())
+	{
+		CurveEditor->SnapToSelectedKey();
+	}
+}
+
 void SCurveEditorPanel::UpdateEditBox()
 {
 	const FCurveEditorSelection& Selection = CurveEditor->Selection;
@@ -890,6 +905,7 @@ TSharedRef<SWidget> SCurveEditorPanel::MakeCurveEditorCurveViewOptionsMenu()
 
 	MenuBuilder.AddMenuSeparator();
 	MenuBuilder.AddMenuEntry(FCurveEditorCommands::Get().ToggleAutoFrameCurveEditor);
+	MenuBuilder.AddMenuEntry(FCurveEditorCommands::Get().ToggleSnapTimeToSelection);
 	MenuBuilder.AddMenuEntry(FCurveEditorCommands::Get().ToggleShowCurveEditorCurveToolTips);
 
 	MenuBuilder.BeginSection("Organize", LOCTEXT("CurveEditorMenuOrganizeHeader", "Organize"));
diff --git a/Engine/Source/Editor/CurveEditor/Public/CurveEditor.h b/Engine/Source/Editor/CurveEditor/Public/CurveEditor.h
--- a/Engine/Source/Editor/CurveEditor/Public/CurveEditor.h
+++ b/Engine/Source/Editor/CurveEditor/Public/CurveEditor.h
@@ -181,6 +181,9 @@ public:
 	void TranslateSelectedKeysLeft();
 	void TranslateSelectedKeysRight();
 
+	/** Snap time to the first selected key */
+	void SnapToSelectedKey();
+
 	/** Step to next or previous key from the current time */
 	void StepToNextKey();
 	void StepToPreviousKey();
diff --git a/Engine/Source/Editor/CurveEditor/Public/CurveEditorCommands.h b/Engine/Source/Editor/CurveEditor/Public/CurveEditorCommands.h
--- a/Engine/Source/Editor/CurveEditor/Public/CurveEditorCommands.h
+++ b/Engine/Source/Editor/CurveEditor/Public/CurveEditorCommands.h
@@ -64,6 +64,7 @@ public:
 	TSharedPtr<FUICommandInfo> SetNoTangentsVisibility;
 
 	TSharedPtr<FUICommandInfo> ToggleAutoFrameCurveEditor;
+	TSharedPtr<FUICommandInfo> ToggleSnapTimeToSelection;
 	TSharedPtr<FUICommandInfo> ToggleShowCurveEditorCurveToolTips;
 
 	TSharedPtr<FUICommandInfo> SetRandomCurveColorsForSelected;
diff --git a/Engine/Source/Editor/CurveEditor/Public/CurveEditorSettings.h b/Engine/Source/Editor/CurveEditor/Public/CurveEditorSettings.h
--- a/Engine/Source/Editor/CurveEditor/Public/CurveEditorSettings.h
+++ b/Engine/Source/Editor/CurveEditor/Public/CurveEditorSettings.h
@@ -102,6 +102,11 @@ public:
 	/** Set zoom in/out position (mouse position or current time). */
 	void SetZoomPosition(ECurveEditorZoomPosition InZoomPosition);
 
+	/** Get whether to snap the time to the currently selected key. */
+	bool GetSnapTimeToSelection() const;
+	/** Set whether to snap the time to the currently selected key. */
+	void SetSnapTimeToSelection(bool bInSnapTimeToSelection);
+
 	/** Get custom color for object and property if it exists, if it doesn't the optional won't be set */
 	TOptional<FLinearColor> GetCustomColor(UClass* InClass, const FString& InPropertyName) const;
 	/** Set Custom Color for the specified parameters. */
@@ -140,6 +145,9 @@ protected:
 	UPROPERTY( config, EditAnywhere, Category="Curve Editor")
 	ECurveEditorZoomPosition ZoomPosition;
 
+	UPROPERTY( config, EditAnywhere, Category="Curve Editor")
+	bool bSnapTimeToSelection;
+
 	UPROPERTY(config, EditAnywhere, Category="Curve Editor")
 	TArray<FCustomColorForChannel> CustomColors;
 
diff --git a/Engine/Source/Editor/CurveEditor/Public/SCurveEditorPanel.h b/Engine/Source/Editor/CurveEditor/Public/SCurveEditorPanel.h
--- a/Engine/Source/Editor/CurveEditor/Public/SCurveEditorPanel.h
+++ b/Engine/Source/Editor/CurveEditor/Public/SCurveEditorPanel.h
@@ -159,6 +159,7 @@ private:
 	EVisibility GetSplitterVisibility() const;
 
 	/*~ Event bindings */
+	void UpdateTime();
 	void UpdateEditBox();
 	void UpdateCommonCurveInfo();
 
