Curve Editor: Add option to Snap Time to Selection

Follow up fix for function change

diff --git a/Engine/Source/Editor/CurveEditor/Private/CurveEditor.cpp b/Engine/Source/Editor/CurveEditor/Private/CurveEditor.cpp
--- a/Engine/Source/Editor/CurveEditor/Private/CurveEditor.cpp
+++ b/Engine/Source/Editor/CurveEditor/Private/CurveEditor.cpp
@@ -794,7 +794,7 @@ void FCurveEditor::SnapToSelectedKey()
 
 	if (MinTime.IsSet())
 	{
-		TimeSliderController->SetScrubPosition(MinTime.GetValue() * TickResolution);		
+		TimeSliderController->SetScrubPosition(MinTime.GetValue() * TickResolution,/*bEvaluate*/ true);		
 	}
 }
 
