Curve Editor: Buffered curve features

Buffered curves now can only operate on themselves, matching by LongDisplayName. The operations are also restricted to the selected curves (either selected in the tree or selected by way of selected keys). There is a new option to Swap Buffered Curves which stores the current curve to the buffer and restores the buffer to the current curve. The Store/Swap/Apply menus are now visible for all the context menus (keys, hovered curve, all curves). There is also a new toggle to show/hide buffered curves and they are also only visible when the curve is selected.

diff --git a/Engine/Source/Editor/CurveEditor/Private/CurveEditor.cpp b/Engine/Source/Editor/CurveEditor/Private/CurveEditor.cpp
--- a/Engine/Source/Editor/CurveEditor/Private/CurveEditor.cpp
+++ b/Engine/Source/Editor/CurveEditor/Private/CurveEditor.cpp
@@ -419,6 +419,11 @@ void FCurveEditor::BindCommands()
 		FIsActionChecked::CreateLambda( [CurveSettings]{ return CurveSettings->GetSnapTimeToSelection(); } )
 	);
 
+	CommandList->MapAction(FCurveEditorCommands::Get().ToggleShowBufferedCurves,
+		FExecuteAction::CreateLambda( [CurveSettings]{ CurveSettings->SetShowBufferedCurves( !CurveSettings->GetShowBufferedCurves() ); } ),
+		FCanExecuteAction(),
+		FIsActionChecked::CreateLambda( [CurveSettings]{ return CurveSettings->GetShowBufferedCurves(); } ) );
+
 	CommandList->MapAction(FCurveEditorCommands::Get().ToggleShowCurveEditorCurveToolTips,
 		FExecuteAction::CreateLambda( [CurveSettings]{ CurveSettings->SetShowCurveEditorCurveToolTips( !CurveSettings->GetShowCurveEditorCurveToolTips() ); } ),
 		FCanExecuteAction(),
@@ -1784,10 +1789,8 @@ TSet<FCurveModelID> FCurveEditor::GetEditedCurves() const
 	return TSet<FCurveModelID>(AllCurves);
 }
 
-void FCurveEditor::SetBufferedCurves(const TSet<FCurveModelID>& InCurves)
+void FCurveEditor::AddBufferedCurves(const TSet<FCurveModelID>& InCurves)
 {
-	BufferedCurves.Empty();
-
 	// We make a copy of the curve data and store it.
 	for (FCurveModelID CurveID : InCurves)
 	{
@@ -1797,7 +1800,20 @@ void FCurveEditor::SetBufferedCurves(const TSet<FCurveModelID>& InCurves)
 		// Add a buffered curve copy if the curve model supports buffered curves
 		TUniquePtr<IBufferedCurveModel> CurveModelCopy = CurveModel->CreateBufferedCurveCopy();
 		if (CurveModelCopy) 
-		{ 
+		{
+			// Remove any existing buffered curves
+			for (int32 BufferedCurveIndex = 0; BufferedCurveIndex < BufferedCurves.Num(); )
+			{
+				if (BufferedCurves[BufferedCurveIndex]->GetLongDisplayName() == CurveModel->GetLongDisplayName().ToString())
+				{
+					BufferedCurves.RemoveAt(BufferedCurveIndex);
+				}
+				else
+				{
+					++BufferedCurveIndex;
+				}
+			}
+
 			BufferedCurves.Add(MoveTemp(CurveModelCopy)); 
 		}
 		else
@@ -1830,9 +1846,9 @@ void FCurveEditor::ApplyBufferedCurveToTarget(const IBufferedCurveModel* Buffere
 	TargetCurve->AddKeys(KeyPositions, KeyAttributes);
 }
 
-bool FCurveEditor::ApplyBufferedCurves(const TSet<FCurveModelID>& InCurvesToApplyTo)
+bool FCurveEditor::ApplyBufferedCurves(const TSet<FCurveModelID>& InCurvesToApplyTo, const bool bSwapBufferCurves)
 {
-	FScopedTransaction Transaction(LOCTEXT("ApplyBufferedCurves", "Apply Buffered Curves"));
+	FScopedTransaction Transaction(bSwapBufferCurves ? LOCTEXT("SwapBufferedCurves", "Swap Buffered Curves") : LOCTEXT("ApplyBufferedCurves", "Apply Buffered Curves"));
 
 	// Each curve can specify an "Intention" name. This gives a little bit of context about how the curve is intended to be used,
 	// without locking anyone into a specific set of intentions. When you go to apply the buffered curves, for each curve that you
@@ -1853,7 +1869,7 @@ bool FCurveEditor::ApplyBufferedCurves(const TSet<FCurveModelID>& InCurvesToAppl
 		check(TargetCurve);
 
 		// Figure out what our destination thinks it's supposed to be used for, ie "Location.X"
-		FString TargetIntent = TargetCurve->GetIntentionName();
+		FString TargetIntent = TargetCurve->GetLongDisplayName().ToString();
 		if (TargetIntent.IsEmpty())
 		{
 			// We don't try to match curves with no intent as that's just chaos.
@@ -1879,7 +1895,7 @@ bool FCurveEditor::ApplyBufferedCurves(const TSet<FCurveModelID>& InCurvesToAppl
 		int32 MatchedBufferedCurveIndex = -1;
 		for (int32 BufferedCurveIndex = BufferedCurveSearchIndexStart; BufferedCurveIndex < BufferedCurves.Num(); BufferedCurveIndex++)
 		{
-			if (BufferedCurves[BufferedCurveIndex]->GetIntentionName() == TargetIntent)
+			if (BufferedCurves[BufferedCurveIndex]->GetLongDisplayName() == TargetIntent)
 			{
 				MatchedBufferedCurveIndex = BufferedCurveIndex;
 
@@ -1905,8 +1921,19 @@ bool FCurveEditor::ApplyBufferedCurves(const TSet<FCurveModelID>& InCurvesToAppl
 			bFoundAnyMatchedIntent = true;
 
 			const IBufferedCurveModel* BufferedCurve = BufferedCurves[MatchedBufferedCurveIndex].Get();
+
+			TUniquePtr<IBufferedCurveModel> CurveModelCopy;
+			if (bSwapBufferCurves)
+			{
+				CurveModelCopy = TargetCurve->CreateBufferedCurveCopy();
+			}
+
 			ApplyBufferedCurveToTarget(BufferedCurve, TargetCurve);
 
+			if (bSwapBufferCurves)
+			{
+				BufferedCurves[MatchedBufferedCurveIndex] = MoveTemp(CurveModelCopy);
+			}
 		}
 		else
 		{
@@ -1956,7 +1983,20 @@ bool FCurveEditor::ApplyBufferedCurves(const TSet<FCurveModelID>& InCurvesToAppl
 		
 		for (int32 CurveIndex = 0; CurveIndex < InCurvesToApplyTo.Num(); CurveIndex++)
 		{
-			ApplyBufferedCurveToTarget(BufferedCurves[CurveIndex].Get(), FindCurve(CurvesToApplyTo[CurveIndex]));
+			FCurveModel* TargetCurve = FindCurve(CurvesToApplyTo[CurveIndex]);
+
+			TUniquePtr<IBufferedCurveModel> CurveModelCopy;
+			if (bSwapBufferCurves)
+			{
+				CurveModelCopy = TargetCurve->CreateBufferedCurveCopy();
+			}
+
+			ApplyBufferedCurveToTarget(BufferedCurves[CurveIndex].Get(), TargetCurve);
+
+			if (bSwapBufferCurves)
+			{
+				BufferedCurves[CurveIndex] = MoveTemp(CurveModelCopy);
+			}
 		}
 
 		FText NotificationText;
@@ -1996,6 +2036,48 @@ bool FCurveEditor::ApplyBufferedCurves(const TSet<FCurveModelID>& InCurvesToAppl
 	return false;
 }
 
+TSet<FCurveModelID> FCurveEditor::GetCurvesForBufferedCurves() const
+{
+	TSet<FCurveModelID> CurveModelIDs;
+
+	// Buffer curves operates on the selected curves (tree selection or key selection)
+	for (const TTuple<FCurveEditorTreeItemID, ECurveEditorTreeSelectionState>& Pair : GetTreeSelection())
+	{
+		if (Pair.Value == ECurveEditorTreeSelectionState::Explicit)
+		{
+			const FCurveEditorTreeItem& TreeItem = GetTreeItem(Pair.Key);
+			for (const FCurveModelID& CurveModelID : TreeItem.GetCurves())
+			{
+				CurveModelIDs.Add(CurveModelID);
+			}
+		}
+	}
+
+	for (const TTuple<FCurveModelID, FKeyHandleSet>& Pair : Selection.GetAll())
+	{
+		CurveModelIDs.Add(Pair.Key);
+	}
+
+	return CurveModelIDs;
+}
+
+bool FCurveEditor::IsActiveBufferedCurve(const TUniquePtr<IBufferedCurveModel>& BufferedCurve) const
+{
+	TSet<FCurveModelID> CurveModelIDs = GetCurvesForBufferedCurves();
+	for (const FCurveModelID& CurveModelID : CurveModelIDs)
+	{
+		if (FCurveModel* Curve = FindCurve(CurveModelID))
+		{
+			if (Curve->GetLongDisplayName().ToString() == BufferedCurve.Get()->GetLongDisplayName())
+			{
+				return true;
+			}
+		}
+	}
+
+	return false;
+}
+
 void FCurveEditor::PostUndo(bool bSuccess)
 {
 	if (WeakPanel.IsValid())
diff --git a/Engine/Source/Editor/CurveEditor/Private/CurveEditorCommands.cpp b/Engine/Source/Editor/CurveEditor/Private/CurveEditorCommands.cpp
--- a/Engine/Source/Editor/CurveEditor/Private/CurveEditorCommands.cpp
+++ b/Engine/Source/Editor/CurveEditor/Private/CurveEditorCommands.cpp
@@ -51,6 +51,7 @@ void FCurveEditorCommands::RegisterCommands()
 
 	UI_COMMAND(ToggleAutoFrameCurveEditor, "Auto Frame Curves", "Auto frame curves when they are selected.", EUserInterfaceActionType::ToggleButton, FInputChord() );
 	UI_COMMAND(ToggleSnapTimeToSelection, "Snap Time to Selection", "Snap the current time to the first selected key time.", EUserInterfaceActionType::ToggleButton, FInputChord() );
+	UI_COMMAND(ToggleShowBufferedCurves, "Buffered Curves", "Show buffered curves for the selected curves.", EUserInterfaceActionType::ToggleButton, FInputChord() );
 	UI_COMMAND(ToggleShowCurveEditorCurveToolTips, "Curve Tool Tips", "Show a tool tip with name and values when hovering over a curve.", EUserInterfaceActionType::ToggleButton, FInputChord() );
 
 	UI_COMMAND(AddKeyHovered, "Add Key", "Add a new key to this curve at the current position.", EUserInterfaceActionType::Button, FInputChord(EKeys::MiddleMouseButton) );
@@ -77,10 +78,11 @@ void FCurveEditorCommands::RegisterCommands()
 	// Deselect any keys that the user has selected.
 	UI_COMMAND(DeselectAllKeys, "Deselect All Keys", "Clears your current key selection.", EUserInterfaceActionType::Button, FInputChord(EModifierKey::Control, EKeys::D));
 
-	// Buffer and Apply Curves. Like copy and paste, but with multiple curve support.
-	UI_COMMAND(BufferVisibleCurves, "Store Curves", "Stores a copy of the visible curves which can be applied onto other curve sets.", EUserInterfaceActionType::Button, FInputChord());
-	// This name is overwritten in CurveEditorContextMenu to show the number of stashed curves.
-	UI_COMMAND(ApplyBufferedCurves, "Apply Stored Curves", "Applies the stored curves to the visible set.", EUserInterfaceActionType::Button, FInputChord());
+	// Buffer and Apply Curves. Like copy and paste, but directly onto the curves they were stored from.
+	// These names are overwritten in CurveEditorContextMenu to show the number of stashed curves.
+	UI_COMMAND(BufferVisibleCurves, "Store Curves", "Stores a copy of the selected curves which can be applied back onto themselves.", EUserInterfaceActionType::Button, FInputChord());
+	UI_COMMAND(SwapBufferedCurves, "Swap Stored Curves", "Applies the stored curves to the curves they were stored from and stores the current curves to the buffer.", EUserInterfaceActionType::Button, FInputChord());
+	UI_COMMAND(ApplyBufferedCurves, "Apply Stored Curves", "Applies the stored curves to the curves they were stored from.", EUserInterfaceActionType::Button, FInputChord());
 
 	// Axis Snapping
 	UI_COMMAND(SetAxisSnappingNone, "Both", "Disable axis snapping and allow movement on both the X and Y directions.", EUserInterfaceActionType::Button, FInputChord());
diff --git a/Engine/Source/Editor/CurveEditor/Private/CurveEditorContextMenu.cpp b/Engine/Source/Editor/CurveEditor/Private/CurveEditorContextMenu.cpp
--- a/Engine/Source/Editor/CurveEditor/Private/CurveEditorContextMenu.cpp
+++ b/Engine/Source/Editor/CurveEditor/Private/CurveEditorContextMenu.cpp
@@ -14,10 +14,21 @@ void FCurveEditorContextMenu::BuildMenu(FMenuBuilder& MenuBuilder, TSharedRef<FC
 
 	TSharedPtr<FCurveEditor> LocalCurveEditor = CurveEditor->AsShared();
 
+	int32 NumActiveCurves = LocalCurveEditor->GetCurvesForBufferedCurves().Num();
+
+	int32 NumBufferedCurves = 0;
+	for (const TUniquePtr<IBufferedCurveModel>& BufferedCurve : LocalCurveEditor->GetBufferedCurves())
+	{
+		if (LocalCurveEditor->IsActiveBufferedCurve(BufferedCurve))
+		{
+			++NumBufferedCurves;
+		}
+	}
+
 	// We change the name to reflect the current number of curves selected.
-	TAttribute<FText> ApplyBufferedCurvesText = TAttribute<FText>::Create(TAttribute<FText>::FGetter::CreateLambda([LocalCurveEditor] {
-		return FText::Format(LOCTEXT("ApplyStoredCurvesContextMenu", "Apply {0} Stored Curves"), LocalCurveEditor->GetNumBufferedCurves());
-	}));
+	TAttribute<FText> BufferedCurvesText = FText::Format(LOCTEXT("StoreCurvesContextMenu", "Store {0} Curves"), NumActiveCurves);
+	TAttribute<FText> SwapBufferedCurvesText = FText::Format(LOCTEXT("SwapStoredCurvesContextMenu", "Swap {0} Stored Curves onto {1} Selected Curves"), NumBufferedCurves, NumActiveCurves);
+	TAttribute<FText> ApplyBufferedCurvesText = FText::Format(LOCTEXT("ApplyStoredCurvesContextMenu", "Apply {0} Stored Curves onto {1} Selected Curves"), NumBufferedCurves, NumActiveCurves);
 
 	const FCurveModel* HoveredCurve = HoveredCurveID.IsSet() ? CurveEditor->FindCurve(HoveredCurveID.GetValue()) : nullptr;
 
@@ -69,6 +80,15 @@ void FCurveEditorContextMenu::BuildMenu(FMenuBuilder& MenuBuilder, TSharedRef<FC
 				MenuBuilder.AddMenuSeparator();
 			}
 
+			// Buffer Curves
+			MenuBuilder.AddMenuEntry(FCurveEditorCommands::Get().BufferVisibleCurves, NAME_None, BufferedCurvesText);
+			if (!bIsReadOnly)
+			{
+				MenuBuilder.AddMenuEntry(FCurveEditorCommands::Get().SwapBufferedCurves, NAME_None, SwapBufferedCurvesText);
+				MenuBuilder.AddMenuEntry(FCurveEditorCommands::Get().ApplyBufferedCurves, NAME_None, ApplyBufferedCurvesText);
+			}
+			MenuBuilder.AddSeparator();
+
 			// Select
 			MenuBuilder.AddMenuEntry(FCurveEditorCommands::Get().SelectAllKeys);
 
@@ -82,18 +102,25 @@ void FCurveEditorContextMenu::BuildMenu(FMenuBuilder& MenuBuilder, TSharedRef<FC
 	}
 	else
 	{
-		if (HoveredCurve)
+		// Test if at least one curve is editable
+		bool bIsReadOnly = true;
+		TSet<FCurveModelID> CurvesToAddTo;
+		for(const FCurveModelID& CurveModelID : CurveEditor->GetEditedCurves())
 		{
-			MenuBuilder.BeginSection("CurveEditorCurveSection", FText::Format(LOCTEXT("CurveNameFormat", "Curve '{0}'"), HoveredCurve->GetLongDisplayName()));
+			if (const FCurveModel* CurveModel = CurveEditor->FindCurve(CurveModelID))
 			{
-				// Buffer Curves
-				MenuBuilder.AddMenuEntry(FCurveEditorCommands::Get().BufferVisibleCurves);
-				if (!HoveredCurve->IsReadOnly())
+				if (!CurveModel->IsReadOnly())
 				{
-					MenuBuilder.AddMenuEntry(FCurveEditorCommands::Get().ApplyBufferedCurves, NAME_None, ApplyBufferedCurvesText);
+					bIsReadOnly = false;
+					break;
 				}
-				MenuBuilder.AddMenuSeparator();
+			}
+		}
 
+		if (HoveredCurve)
+		{
+			MenuBuilder.BeginSection("CurveEditorCurveSection", FText::Format(LOCTEXT("CurveNameFormat", "Curve '{0}'"), HoveredCurve->GetLongDisplayName()));
+			{
 				// Modify Curve
 				if (!HoveredCurve->IsReadOnly())
 				{
@@ -125,6 +152,15 @@ void FCurveEditorContextMenu::BuildMenu(FMenuBuilder& MenuBuilder, TSharedRef<FC
 					);
 
 					MenuBuilder.AddMenuSeparator();
+
+					// Buffer Curves
+					MenuBuilder.AddMenuEntry(FCurveEditorCommands::Get().BufferVisibleCurves, NAME_None, BufferedCurvesText);
+					if (!bIsReadOnly)
+					{
+						MenuBuilder.AddMenuEntry(FCurveEditorCommands::Get().SwapBufferedCurves, NAME_None, SwapBufferedCurvesText);
+						MenuBuilder.AddMenuEntry(FCurveEditorCommands::Get().ApplyBufferedCurves, NAME_None, ApplyBufferedCurvesText);
+					}
+					MenuBuilder.AddSeparator();
 				}
 
 				// Select
@@ -140,33 +176,23 @@ void FCurveEditorContextMenu::BuildMenu(FMenuBuilder& MenuBuilder, TSharedRef<FC
 		}
 		else
 		{
-			// Test if at least one curve is editable
-			bool bIsReadOnly = true;
-			TSet<FCurveModelID> CurvesToAddTo;
-			for(const FCurveModelID& CurveModelID : CurveEditor->GetEditedCurves())
-			{
-				if (const FCurveModel* CurveModel = CurveEditor->FindCurve(CurveModelID))
-				{
-					if (!CurveModel->IsReadOnly())
-					{
-						bIsReadOnly = false;
-						break;
-					}
-				}
-			}
-
 			MenuBuilder.BeginSection("CurveEditorAllCurveSections", LOCTEXT("CurveEditorAllCurveSections", "All Curves"));
 			{
-				// Buffer Curves
-				MenuBuilder.AddMenuEntry(FCurveEditorCommands::Get().BufferVisibleCurves);
 				if (!bIsReadOnly)
 				{
-					MenuBuilder.AddMenuEntry(FCurveEditorCommands::Get().ApplyBufferedCurves, NAME_None, ApplyBufferedCurvesText);
-					MenuBuilder.AddMenuSeparator();
-
 					// Modify Curves
 					MenuBuilder.AddMenuEntry(FCurveEditorCommands::Get().AddKeyToAllCurves);
-					MenuBuilder.AddMenuSeparator();
+
+					MenuBuilder.AddSeparator();
+
+					// Buffer Curves
+					MenuBuilder.AddMenuEntry(FCurveEditorCommands::Get().BufferVisibleCurves, NAME_None, BufferedCurvesText);
+					if (!bIsReadOnly)
+					{
+						MenuBuilder.AddMenuEntry(FCurveEditorCommands::Get().SwapBufferedCurves, NAME_None, SwapBufferedCurvesText);
+						MenuBuilder.AddMenuEntry(FCurveEditorCommands::Get().ApplyBufferedCurves, NAME_None, ApplyBufferedCurvesText);
+					}
+					MenuBuilder.AddSeparator();
 				}
 
 				// Select
diff --git a/Engine/Source/Editor/CurveEditor/Private/CurveEditorSettings.cpp b/Engine/Source/Editor/CurveEditor/Private/CurveEditorSettings.cpp
--- a/Engine/Source/Editor/CurveEditor/Private/CurveEditorSettings.cpp
+++ b/Engine/Source/Editor/CurveEditor/Private/CurveEditorSettings.cpp
@@ -7,6 +7,7 @@ UCurveEditorSettings::UCurveEditorSettings()
 	bAutoFrameCurveEditor = true;
 	FrameInputPadding = 50;
 	FrameOutputPadding = 50;
+	bShowBufferedCurves = true;
 	bShowCurveEditorCurveToolTips = true;
 	TangentVisibility = ECurveEditorTangentVisibility::SelectedKeys;
 	ZoomPosition = ECurveEditorZoomPosition::CurrentTime;
@@ -58,6 +59,20 @@ void UCurveEditorSettings::SetFrameOutputPadding(int32 InFrameOutputPadding)
 	}
 }
 
+bool UCurveEditorSettings::GetShowBufferedCurves() const
+{
+	return bShowBufferedCurves;
+}
+
+void UCurveEditorSettings::SetShowBufferedCurves(bool InbShowBufferedCurves)
+{
+	if (bShowBufferedCurves != InbShowBufferedCurves)
+	{
+		bShowBufferedCurves = InbShowBufferedCurves;
+		SaveConfig();
+	}
+}
+
 bool UCurveEditorSettings::GetShowCurveEditorCurveToolTips() const
 {
 	return bShowCurveEditorCurveToolTips;
diff --git a/Engine/Source/Editor/CurveEditor/Private/RichCurveEditorModel.cpp b/Engine/Source/Editor/CurveEditor/Private/RichCurveEditorModel.cpp
--- a/Engine/Source/Editor/CurveEditor/Private/RichCurveEditorModel.cpp
+++ b/Engine/Source/Editor/CurveEditor/Private/RichCurveEditorModel.cpp
@@ -63,8 +63,8 @@ class FRichBufferedCurveModel : public IBufferedCurveModel
 {
 public:
 	FRichBufferedCurveModel(const FRichCurve& InRichCurve, TArray<FKeyPosition>&& InKeyPositions, TArray<FKeyAttributes>&& InKeyAttributes,
-		const FString& InIntentionName, const double InValueMin, const double InValueMax)
-		: IBufferedCurveModel(MoveTemp(InKeyPositions), MoveTemp(InKeyAttributes), InIntentionName, InValueMin, InValueMax)
+		const FString& InLongDisplayName, const double InValueMin, const double InValueMax)
+		: IBufferedCurveModel(MoveTemp(InKeyPositions), MoveTemp(InKeyAttributes), InLongDisplayName, InValueMin, InValueMax)
 		, RichCurve(InRichCurve)
 	{}
 
@@ -634,7 +634,7 @@ TUniquePtr<IBufferedCurveModel> FRichCurveEditorModel::CreateBufferedCurveCopy()
 			double ValueMin = 0.f, ValueMax = 1.f;
 			GetValueRange(ValueMin, ValueMax);
 
-			return MakeUnique<FRichBufferedCurveModel>(RichCurve, MoveTemp(KeyPositions), MoveTemp(KeyAttributes), GetIntentionName(), ValueMin, ValueMax);
+			return MakeUnique<FRichBufferedCurveModel>(RichCurve, MoveTemp(KeyPositions), MoveTemp(KeyAttributes), GetLongDisplayName().ToString(), ValueMin, ValueMax);
 		}
 	}
 
diff --git a/Engine/Source/Editor/CurveEditor/Private/SCurveEditorPanel.cpp b/Engine/Source/Editor/CurveEditor/Private/SCurveEditorPanel.cpp
--- a/Engine/Source/Editor/CurveEditor/Private/SCurveEditorPanel.cpp
+++ b/Engine/Source/Editor/CurveEditor/Private/SCurveEditorPanel.cpp
@@ -906,6 +906,7 @@ TSharedRef<SWidget> SCurveEditorPanel::MakeCurveEditorCurveViewOptionsMenu()
 	MenuBuilder.AddMenuSeparator();
 	MenuBuilder.AddMenuEntry(FCurveEditorCommands::Get().ToggleAutoFrameCurveEditor);
 	MenuBuilder.AddMenuEntry(FCurveEditorCommands::Get().ToggleSnapTimeToSelection);
+	MenuBuilder.AddMenuEntry(FCurveEditorCommands::Get().ToggleShowBufferedCurves);
 	MenuBuilder.AddMenuEntry(FCurveEditorCommands::Get().ToggleShowCurveEditorCurveToolTips);
 
 	MenuBuilder.BeginSection("Organize", LOCTEXT("CurveEditorMenuOrganizeHeader", "Organize"));
diff --git a/Engine/Source/Editor/CurveEditor/Private/Views/SCurveEditorViewNormalized.cpp b/Engine/Source/Editor/CurveEditor/Private/Views/SCurveEditorViewNormalized.cpp
--- a/Engine/Source/Editor/CurveEditor/Private/Views/SCurveEditorViewNormalized.cpp
+++ b/Engine/Source/Editor/CurveEditor/Private/Views/SCurveEditorViewNormalized.cpp
@@ -2,6 +2,7 @@
 
 #include "Views/SCurveEditorViewNormalized.h"
 #include "CurveEditor.h"
+#include "CurveEditorSettings.h"
 #include "CurveEditorHelpers.h"
 #include "CurveModel.h"
 #include "Widgets/Text/STextBlock.h"
@@ -61,6 +62,11 @@ void SCurveEditorViewNormalized::DrawBufferedCurves(const FGeometry& AllottedGeo
 		return;
 	}
 
+	if (!CurveEditor->GetSettings()->GetShowBufferedCurves())
+	{
+		return;
+	}
+
 	const float BufferedCurveThickness = 1.f;
 	const bool  bAntiAliasCurves = true;
 	const FLinearColor CurveColor = CurveViewConstants::BufferedCurveColor;
@@ -73,6 +79,11 @@ void SCurveEditorViewNormalized::DrawBufferedCurves(const FGeometry& AllottedGeo
 	const TArray<TUniquePtr<IBufferedCurveModel>>& BufferedCurves = CurveEditor->GetBufferedCurves();
 	for (const TUniquePtr<IBufferedCurveModel>& BufferedCurve : BufferedCurves)
 	{
+		if (!CurveEditor->IsActiveBufferedCurve(BufferedCurve))
+		{
+			continue;
+		}
+
 		FTransform2D ViewToBufferedCurveTransform;
 		double CurveOutputMin = BufferedCurve->GetValueMin(), CurveOutputMax = BufferedCurve->GetValueMax();
 
diff --git a/Engine/Source/Editor/CurveEditor/Private/Views/SCurveEditorViewStacked.cpp b/Engine/Source/Editor/CurveEditor/Private/Views/SCurveEditorViewStacked.cpp
--- a/Engine/Source/Editor/CurveEditor/Private/Views/SCurveEditorViewStacked.cpp
+++ b/Engine/Source/Editor/CurveEditor/Private/Views/SCurveEditorViewStacked.cpp
@@ -2,6 +2,7 @@
 
 #include "Views/SCurveEditorViewStacked.h"
 #include "CurveEditor.h"
+#include "CurveEditorSettings.h"
 #include "CurveModel.h"
 #include "SCurveEditorPanel.h"
 
@@ -279,6 +280,11 @@ void SCurveEditorViewStacked::DrawBufferedCurves(const FGeometry& AllottedGeomet
 		return;
 	}
 
+	if (!CurveEditor->GetSettings()->GetShowBufferedCurves())
+	{
+		return;
+	}
+
 	const TArray<TUniquePtr<IBufferedCurveModel>>& BufferedCurves = CurveEditor->GetBufferedCurves();
 
 	const float BufferedCurveThickness = 1.f;
@@ -308,6 +314,11 @@ void SCurveEditorViewStacked::DrawBufferedCurves(const FGeometry& AllottedGeomet
 		// Calculate the view to curve transform for each buffered curve, then draw
 		for (const TUniquePtr<IBufferedCurveModel>& BufferedCurve : BufferedCurves)
 		{
+			if (!CurveEditor->IsActiveBufferedCurve(BufferedCurve))
+			{
+				continue;
+			}
+
 			double CurveOutputMin = BufferedCurve->GetValueMin(), CurveOutputMax = BufferedCurve->GetValueMax();
 
 			ViewToBufferedCurveTransform = CalculateViewToCurveTransform(CurveOutputMin, CurveOutputMax, ValueOffset);
diff --git a/Engine/Source/Editor/CurveEditor/Private/Views/SInteractiveCurveEditorView.cpp b/Engine/Source/Editor/CurveEditor/Private/Views/SInteractiveCurveEditorView.cpp
--- a/Engine/Source/Editor/CurveEditor/Private/Views/SInteractiveCurveEditorView.cpp
+++ b/Engine/Source/Editor/CurveEditor/Private/Views/SInteractiveCurveEditorView.cpp
@@ -449,6 +449,11 @@ void SInteractiveCurveEditorView::DrawCurves(TSharedRef<FCurveEditor> CurveEdito
 
 void SInteractiveCurveEditorView::DrawBufferedCurves(TSharedRef<FCurveEditor> CurveEditor, const FGeometry& AllottedGeometry, const FSlateRect& MyCullingRect, FSlateWindowElementList& OutDrawElements, int32 BaseLayerId, const FWidgetStyle& InWidgetStyle, ESlateDrawEffect DrawEffects) const
 {
+	if (!CurveEditor->GetSettings()->GetShowBufferedCurves())
+	{
+		return;
+	}
+
 	const float BufferedCurveThickness = 1.f;
 	const bool  bAntiAliasCurves = true;
 	const FLinearColor CurveColor = CurveViewConstants::BufferedCurveColor;
@@ -459,6 +464,11 @@ void SInteractiveCurveEditorView::DrawBufferedCurves(TSharedRef<FCurveEditor> Cu
 	// Draw each buffered curve using the view space transform since the curve space for all curves is the same
 	for (const TUniquePtr<IBufferedCurveModel>& BufferedCurve : BufferedCurves)
 	{
+		if (!CurveEditor->IsActiveBufferedCurve(BufferedCurve))
+		{
+			continue;
+		}
+
 		TArray<TTuple<double, double>> CurveSpaceInterpolatingPoints;
 		FCurveEditorScreenSpace CurveSpace = GetViewSpace();
 
@@ -1299,6 +1309,7 @@ void SInteractiveCurveEditorView::RebindContextualActions(FVector2D InMousePosit
 	CommandList->UnmapAction(FCurveEditorCommands::Get().AddKeyToAllCurves);
 
 	CommandList->UnmapAction(FCurveEditorCommands::Get().BufferVisibleCurves);
+	CommandList->UnmapAction(FCurveEditorCommands::Get().SwapBufferedCurves);
 	CommandList->UnmapAction(FCurveEditorCommands::Get().ApplyBufferedCurves);
 	
 
@@ -1310,90 +1321,51 @@ void SInteractiveCurveEditorView::RebindContextualActions(FVector2D InMousePosit
 
 		CommandList->MapAction(FCurveEditorCommands::Get().AddKeyHovered, FExecuteAction::CreateSP(this, &SInteractiveCurveEditorView::AddKeyAtMousePosition, HoveredCurveSet));
 		CommandList->MapAction(FCurveEditorCommands::Get().PasteKeysHovered, FExecuteAction::CreateSP(this, &SInteractiveCurveEditorView::PasteKeys, HoveredCurveSet));
-
-		// Buffer the curve they have highlighted instead of all of them.
-		CommandList->MapAction(FCurveEditorCommands::Get().BufferVisibleCurves, FExecuteAction::CreateSP(this, &SInteractiveCurveEditorView::BufferCurve, HoveredCurve.GetValue()));
-	}
-	else
-	{
-		// Apply the buffering action to our entire set and not just the hovered curve.
-		CommandList->MapAction(FCurveEditorCommands::Get().BufferVisibleCurves, FExecuteAction::CreateSP(this, &SInteractiveCurveEditorView::BufferVisibleCurves));
 	}
 
 	CommandList->MapAction(FCurveEditorCommands::Get().AddKeyToAllCurves, FExecuteAction::CreateSP(this, &SInteractiveCurveEditorView::AddKeyAtScrubTime, TSet<FCurveModelID>()));
 
-	// Buffer Visible Curves. Can only apply buffered curves if the current number of visible curves matches the number of buffered curves.
-	CommandList->MapAction(FCurveEditorCommands::Get().ApplyBufferedCurves, FExecuteAction::CreateSP(this, &SInteractiveCurveEditorView::ApplyBufferCurves, HoveredCurve), FCanExecuteAction::CreateSP(this, &SInteractiveCurveEditorView::CanApplyBufferedCurves, HoveredCurve));
+	// Buffer Curves. Can only act on buffered curves if curves are selected in the tree or the curve has selected keys.
+	CommandList->MapAction(FCurveEditorCommands::Get().BufferVisibleCurves, FExecuteAction::CreateSP(this, &SInteractiveCurveEditorView::BufferCurves), FCanExecuteAction::CreateSP(this, &SInteractiveCurveEditorView::CanBufferedCurves));
+	CommandList->MapAction(FCurveEditorCommands::Get().SwapBufferedCurves, FExecuteAction::CreateSP(this, &SInteractiveCurveEditorView::ApplyBufferCurves, true), FCanExecuteAction::CreateSP(this, &SInteractiveCurveEditorView::CanApplyBufferedCurves));
+	CommandList->MapAction(FCurveEditorCommands::Get().ApplyBufferedCurves, FExecuteAction::CreateSP(this, &SInteractiveCurveEditorView::ApplyBufferCurves, false), FCanExecuteAction::CreateSP(this, &SInteractiveCurveEditorView::CanApplyBufferedCurves));
 }
 
-void SInteractiveCurveEditorView::BufferVisibleCurves()
+void SInteractiveCurveEditorView::BufferCurves()
 {
 	TSharedPtr<FCurveEditor> CurveEditor = WeakCurveEditor.Pin();
 	if (CurveEditor.IsValid())
 	{
-		// Curve Editor will handle copying and storing the curves.
-		TSet<FCurveModelID> ActiveCurveIDs;
-		for (const TTuple<FCurveModelID, FCurveInfo>& Pair : CurveInfoByID)
-		{
-			ActiveCurveIDs.Add(Pair.Key);
-		}
-		CurveEditor->SetBufferedCurves(ActiveCurveIDs);
+		CurveEditor->AddBufferedCurves(CurveEditor->GetCurvesForBufferedCurves());
 	}
 }
 
-void SInteractiveCurveEditorView::BufferCurve(const FCurveModelID CurveID)
+void SInteractiveCurveEditorView::ApplyBufferCurves(const bool bSwapBufferCurves)
 {
 	TSharedPtr<FCurveEditor> CurveEditor = WeakCurveEditor.Pin();
 	if (CurveEditor.IsValid())
 	{
-		// Curve Editor will handle copying and storing the curves.
-		TSet<FCurveModelID> CurveSet;
-		CurveSet.Add(CurveID);
-		CurveEditor->SetBufferedCurves(CurveSet);
+		CurveEditor->ApplyBufferedCurves(CurveEditor->GetCurvesForBufferedCurves(), bSwapBufferCurves);
 	}
 }
 
-void SInteractiveCurveEditorView::ApplyBufferCurves(TOptional<FCurveModelID> DestinationCurve)
+bool SInteractiveCurveEditorView::CanBufferedCurves() const
 {
 	TSharedPtr<FCurveEditor> CurveEditor = WeakCurveEditor.Pin();
 	if (CurveEditor.IsValid())
 	{
-		if (DestinationCurve.IsSet())
-		{
-			TSet<FCurveModelID> CurveSet;
-			CurveSet.Add(DestinationCurve.GetValue());
-
-			// Apply the buffered curve (singular) to our highlighted curve.
-			CurveEditor->ApplyBufferedCurves(CurveSet);
-		}
-		else
-		{
-			// Curve Editor will handle attempting to apply the buffered curves to our currently visible ones.
-			TSet<FCurveModelID> ActiveCurveIDs;
-			for (const TTuple<FCurveModelID, FCurveInfo>& Pair : CurveInfoByID)
-			{
-				ActiveCurveIDs.Add(Pair.Key);
-			}
-			CurveEditor->ApplyBufferedCurves(ActiveCurveIDs);
-		}
-
+		return CurveEditor->GetCurvesForBufferedCurves().Num() > 0;
 	}
+
+	return false;
 }
 
-bool SInteractiveCurveEditorView::CanApplyBufferedCurves(TOptional<FCurveModelID> DestinationCurve) const
+bool SInteractiveCurveEditorView::CanApplyBufferedCurves() const
 {
 	TSharedPtr<FCurveEditor> CurveEditor = WeakCurveEditor.Pin();
 	if (CurveEditor.IsValid())
 	{
-		if (DestinationCurve.IsSet())
-		{
-			return CurveEditor->GetNumBufferedCurves() == 1;
-		}
-		else
-		{
-			// For now we just do a 1:1 mapping. Once curves have better names we can try to do an intelligent match up, ie: matching Transform.X to a new Transform.X
-			return CurveEditor->GetNumBufferedCurves() == NumCurves();
-		}
+		return CurveEditor->GetCurvesForBufferedCurves().Num() > 0 && CurveEditor->GetBufferedCurves().Num() > 0;
 	}
 
 	return false;
diff --git a/Engine/Source/Editor/CurveEditor/Public/CurveEditor.h b/Engine/Source/Editor/CurveEditor/Public/CurveEditor.h
--- a/Engine/Source/Editor/CurveEditor/Public/CurveEditor.h
+++ b/Engine/Source/Editor/CurveEditor/Public/CurveEditor.h
@@ -274,13 +274,18 @@ public:
 	/** Attribute used for determining default attributes to apply to a newly create key */
 	TAttribute<FKeyAttributes> GetDefaultKeyAttribute() const { return DefaultKeyAttributes; }
 	/** Create a copy of the specified set of curves. */
-	void SetBufferedCurves(const TSet<FCurveModelID>& InCurves);
+	void AddBufferedCurves(const TSet<FCurveModelID>& InCurves);
 	/** Attempts to apply the buffered curve to the passed in curve set. Returns true on success. */
-	bool ApplyBufferedCurves(const TSet<FCurveModelID>& InCurvesToApplyTo);
+	bool ApplyBufferedCurves(const TSet<FCurveModelID>& InCurvesToApplyTo, const bool bSwapBufferCurves);
 	/** Return the number of stored Buffered Curves. */
 	int32 GetNumBufferedCurves() const { return BufferedCurves.Num(); }
 	/** Return the array of buffered curves */
 	const TArray<TUniquePtr<IBufferedCurveModel>>& GetBufferedCurves() const { return BufferedCurves; }
+	/** Return the curve model IDs to operate buffer curves on, ie. selected in the tree view or with selected keys */
+	TSet<FCurveModelID> GetCurvesForBufferedCurves() const;
+	/** Returns whether the buffered curve is to be acted on, ie. selected, in the tree view or with selected keys */
+	bool IsActiveBufferedCurve(const TUniquePtr<IBufferedCurveModel>& BufferedCurve) const;
+
 	// ~FCurveEditor
 
 	// FEditorUndoClient
diff --git a/Engine/Source/Editor/CurveEditor/Public/CurveEditorCommands.h b/Engine/Source/Editor/CurveEditor/Public/CurveEditorCommands.h
--- a/Engine/Source/Editor/CurveEditor/Public/CurveEditorCommands.h
+++ b/Engine/Source/Editor/CurveEditor/Public/CurveEditorCommands.h
@@ -65,6 +65,7 @@ public:
 
 	TSharedPtr<FUICommandInfo> ToggleAutoFrameCurveEditor;
 	TSharedPtr<FUICommandInfo> ToggleSnapTimeToSelection;
+	TSharedPtr<FUICommandInfo> ToggleShowBufferedCurves;
 	TSharedPtr<FUICommandInfo> ToggleShowCurveEditorCurveToolTips;
 
 	TSharedPtr<FUICommandInfo> SetRandomCurveColorsForSelected;
@@ -85,6 +86,7 @@ public:
 	TSharedPtr<FUICommandInfo> DeselectAllKeys;
 
 	TSharedPtr<FUICommandInfo> BufferVisibleCurves;
+	TSharedPtr<FUICommandInfo> SwapBufferedCurves;
 	TSharedPtr<FUICommandInfo> ApplyBufferedCurves;
 
 	// User Filtering
diff --git a/Engine/Source/Editor/CurveEditor/Public/CurveEditorSettings.h b/Engine/Source/Editor/CurveEditor/Public/CurveEditorSettings.h
--- a/Engine/Source/Editor/CurveEditor/Public/CurveEditorSettings.h
+++ b/Engine/Source/Editor/CurveEditor/Public/CurveEditorSettings.h
@@ -87,6 +87,11 @@ public:
 	/** Sets the number of pixels to pad output framing */
 	void SetFrameOutputPadding(int32 InFrameOutputPadding);
 
+	/** Gets whether or not to show buffered curves in the curve editor. */
+	bool GetShowBufferedCurves() const;
+	/** Sets whether or not to show buffered curves in the curve editor. */
+	void SetShowBufferedCurves(bool InbShowBufferedCurves);
+
 	/** Gets whether or not to show curve tool tips in the curve editor. */
 	bool GetShowCurveEditorCurveToolTips() const;
 	/** Sets whether or not to show curve tool tips in the curve editor. */
@@ -136,6 +141,9 @@ protected:
 	UPROPERTY( config, EditAnywhere, Category="Curve Editor", meta=(ClampMin=0) )
 	int32 FrameOutputPadding;
 
+	UPROPERTY( config, EditAnywhere, Category="Curve Editor" )
+	bool bShowBufferedCurves;
+
 	UPROPERTY( config, EditAnywhere, Category="Curve Editor" )
 	bool bShowCurveEditorCurveToolTips;
 
diff --git a/Engine/Source/Editor/CurveEditor/Public/IBufferedCurveModel.h b/Engine/Source/Editor/CurveEditor/Public/IBufferedCurveModel.h
--- a/Engine/Source/Editor/CurveEditor/Public/IBufferedCurveModel.h
+++ b/Engine/Source/Editor/CurveEditor/Public/IBufferedCurveModel.h
@@ -16,10 +16,10 @@ class FCurveModel;
 class IBufferedCurveModel
 {
 public:
-	IBufferedCurveModel(TArray<FKeyPosition>&& InKeyPositions, TArray<FKeyAttributes>&& InKeyAttributes, const FString& InIntentionName, const double InValueMin, const double InValueMax) 
+	IBufferedCurveModel(TArray<FKeyPosition>&& InKeyPositions, TArray<FKeyAttributes>&& InKeyAttributes, const FString& InLongDisplayName, const double InValueMin, const double InValueMax) 
 	: KeyPositions(MoveTemp(InKeyPositions))
 	, KeyAttributes(MoveTemp(InKeyAttributes))
-	, IntentionName(InIntentionName)
+	, LongDisplayName(InLongDisplayName)
 	, ValueMin(InValueMin)
 	, ValueMax(InValueMax)
 	{}
@@ -44,11 +44,9 @@ public:
 	FORCEINLINE void GetKeyAttributes(TArray<FKeyAttributes>& OutKeyAttributes) const { OutKeyAttributes = KeyAttributes; }
 
 	/**
-	 * Returns the intention name for this buffered curve. This is an internal name used to try to match 
-	 * different curves with each other. See the CurveModel's GetIntentionName method for more detail on
-	 * what an intention name is
+	 * Returns the full/long display name of the curves. Buffered curves can only be applied to the original curve it was stored from.
 	 */
-	FORCEINLINE FString GetIntentionName() const { return IntentionName; }
+	FORCEINLINE FString GetLongDisplayName() const { return LongDisplayName; }
 
 	/** returns the lowest output value in curve space for this buffered curve */
 	FORCEINLINE double GetValueMin() const { return ValueMin; }
@@ -59,6 +57,6 @@ public:
 protected:
 	TArray<FKeyPosition> KeyPositions;
 	TArray<FKeyAttributes> KeyAttributes;
-	FString IntentionName;
+	FString LongDisplayName;
 	double ValueMin, ValueMax;
 };
\ No newline at end of file
diff --git a/Engine/Source/Editor/CurveEditor/Public/Views/SInteractiveCurveEditorView.h b/Engine/Source/Editor/CurveEditor/Public/Views/SInteractiveCurveEditorView.h
--- a/Engine/Source/Editor/CurveEditor/Public/Views/SInteractiveCurveEditorView.h
+++ b/Engine/Source/Editor/CurveEditor/Public/Views/SInteractiveCurveEditorView.h
@@ -153,13 +153,13 @@ private:
 	void RebindContextualActions(FVector2D InMousePosition);
 
 	/** Copy the curves from this view and set them as the Curve Editor's buffered curve support. */
-	void BufferVisibleCurves();
-	/** Copy the curves from this view and set them as the Curve Editor's buffered curve support. */
-	void BufferCurve(const FCurveModelID CurveID);
-	/** Attempt to apply the previously buffered curves to the currently visible curves. */
-	void ApplyBufferCurves(TOptional<FCurveModelID> DestinationCurve);
-	/** Check if it's legal to apply any of the buffered curves to our currently visible curves. */
-	bool CanApplyBufferedCurves(TOptional<FCurveModelID> DestinationCurve) const;
+	void BufferCurves();
+	/** Attempt to apply the previously buffered curves to the currently selected curves. */
+	void ApplyBufferCurves(const bool bSwapBufferCurves);
+	/** Check if it's legal to buffer any of our selected curves. */
+	bool CanBufferedCurves() const;
+	/** Check if it's legal to apply any of the buffered curves to our currently selected curves. */
+	bool CanApplyBufferedCurves() const;
 
 protected:
 
diff --git a/Engine/Source/Editor/MovieSceneTools/Private/Channels/BezierChannelCurveModel.cpp b/Engine/Source/Editor/MovieSceneTools/Private/Channels/BezierChannelCurveModel.cpp
--- a/Engine/Source/Editor/MovieSceneTools/Private/Channels/BezierChannelCurveModel.cpp
+++ b/Engine/Source/Editor/MovieSceneTools/Private/Channels/BezierChannelCurveModel.cpp
@@ -21,8 +21,8 @@
 template<typename ChannelType>
 FBezierChannelBufferedCurveModel<ChannelType>::FBezierChannelBufferedCurveModel(
 		const ChannelType* InChannel, TWeakObjectPtr<UMovieSceneSection> InWeakSection,
-		TArray<FKeyPosition>&& InKeyPositions, TArray<FKeyAttributes>&& InKeyAttributes, const FString& InIntentionName, const double InValueMin, const double InValueMax)
-	: IBufferedCurveModel(MoveTemp(InKeyPositions), MoveTemp(InKeyAttributes), InIntentionName, InValueMin, InValueMax)
+		TArray<FKeyPosition>&& InKeyPositions, TArray<FKeyAttributes>&& InKeyAttributes, const FString& InLongDisplayName, const double InValueMin, const double InValueMax)
+	: IBufferedCurveModel(MoveTemp(InKeyPositions), MoveTemp(InKeyAttributes), InLongDisplayName, InValueMin, InValueMax)
 	, Channel(*InChannel)
 	, WeakSection(InWeakSection)
 {
diff --git a/Engine/Source/Editor/MovieSceneTools/Private/Channels/BoolChannelCurveModel.cpp b/Engine/Source/Editor/MovieSceneTools/Private/Channels/BoolChannelCurveModel.cpp
--- a/Engine/Source/Editor/MovieSceneTools/Private/Channels/BoolChannelCurveModel.cpp
+++ b/Engine/Source/Editor/MovieSceneTools/Private/Channels/BoolChannelCurveModel.cpp
@@ -26,8 +26,8 @@ class FBoolChannelBufferedCurveModel : public IBufferedCurveModel
 public:
 	/** Create a copy of the float channel while keeping the reference to the section */
 	FBoolChannelBufferedCurveModel(const FMovieSceneBoolChannel* InMovieSceneBoolChannel, TWeakObjectPtr<UMovieSceneSection> InWeakSection,
-		TArray<FKeyPosition>&& InKeyPositions, TArray<FKeyAttributes>&& InKeyAttributes, const FString& InIntentionName, const double InValueMin, const double InValueMax)
-		: IBufferedCurveModel(MoveTemp(InKeyPositions), MoveTemp(InKeyAttributes), InIntentionName, InValueMin, InValueMax)
+		TArray<FKeyPosition>&& InKeyPositions, TArray<FKeyAttributes>&& InKeyAttributes, const FString& InLongDisplayName, const double InValueMin, const double InValueMax)
+		: IBufferedCurveModel(MoveTemp(InKeyPositions), MoveTemp(InKeyAttributes), InLongDisplayName, InValueMin, InValueMax)
 		, Channel(*InMovieSceneBoolChannel)
 		, WeakSection(InWeakSection)
 	{}
@@ -102,7 +102,7 @@ TUniquePtr<IBufferedCurveModel> FBoolChannelCurveModel::CreateBufferedCurveCopy(
 		double ValueMin = 0.f, ValueMax = 1.f;
 		GetValueRange(ValueMin, ValueMax);
 
-		return MakeUnique<FBoolChannelBufferedCurveModel>(Channel, Cast<UMovieSceneSection>(GetOwningObject()), MoveTemp(KeyPositions), MoveTemp(KeyAttributes), GetIntentionName(), ValueMin, ValueMax);
+		return MakeUnique<FBoolChannelBufferedCurveModel>(Channel, Cast<UMovieSceneSection>(GetOwningObject()), MoveTemp(KeyPositions), MoveTemp(KeyAttributes), GetLongDisplayName().ToString(), ValueMin, ValueMax);
 	}
 	return nullptr;
 }
diff --git a/Engine/Source/Editor/MovieSceneTools/Private/Channels/DoubleChannelCurveModel.cpp b/Engine/Source/Editor/MovieSceneTools/Private/Channels/DoubleChannelCurveModel.cpp
--- a/Engine/Source/Editor/MovieSceneTools/Private/Channels/DoubleChannelCurveModel.cpp
+++ b/Engine/Source/Editor/MovieSceneTools/Private/Channels/DoubleChannelCurveModel.cpp
@@ -12,8 +12,8 @@ class FDoubleChannelBufferedCurveModel : public IBufferedCurveModel
 public:
 	/** Create a copy of the double channel while keeping the reference to the section */
 	FDoubleChannelBufferedCurveModel(const FMovieSceneDoubleChannel* InMovieSceneDoubleChannel, TWeakObjectPtr<UMovieSceneSection> InWeakSection,
-		TArray<FKeyPosition>&& InKeyPositions, TArray<FKeyAttributes>&& InKeyAttributes, const FString& InIntentionName, const double InValueMin, const double InValueMax)
-		: IBufferedCurveModel(MoveTemp(InKeyPositions), MoveTemp(InKeyAttributes), InIntentionName, InValueMin, InValueMax)
+		TArray<FKeyPosition>&& InKeyPositions, TArray<FKeyAttributes>&& InKeyAttributes, const FString& InLongDisplayName, const double InValueMin, const double InValueMax)
+		: IBufferedCurveModel(MoveTemp(InKeyPositions), MoveTemp(InKeyAttributes), InLongDisplayName, InValueMin, InValueMax)
 		, Channel(*InMovieSceneDoubleChannel)
 		, WeakSection(InWeakSection)
 	{}
@@ -77,7 +77,7 @@ TUniquePtr<IBufferedCurveModel> FDoubleChannelCurveModel::CreateBufferedCurveCop
 		double ValueMin = 0.f, ValueMax = 1.f;
 		GetValueRange(ValueMin, ValueMax);
 
-		return MakeUnique<FDoubleChannelBufferedCurveModel>(Channel, Cast<UMovieSceneSection>(GetOwningObject()), MoveTemp(KeyPositions), MoveTemp(KeyAttributes), GetIntentionName(), ValueMin, ValueMax);
+		return MakeUnique<FDoubleChannelBufferedCurveModel>(Channel, Cast<UMovieSceneSection>(GetOwningObject()), MoveTemp(KeyPositions), MoveTemp(KeyAttributes), GetLongDisplayName().ToString(), ValueMin, ValueMax);
 	}
 	return nullptr;
 }
diff --git a/Engine/Source/Editor/MovieSceneTools/Private/Channels/FloatChannelCurveModel.cpp b/Engine/Source/Editor/MovieSceneTools/Private/Channels/FloatChannelCurveModel.cpp
--- a/Engine/Source/Editor/MovieSceneTools/Private/Channels/FloatChannelCurveModel.cpp
+++ b/Engine/Source/Editor/MovieSceneTools/Private/Channels/FloatChannelCurveModel.cpp
@@ -12,8 +12,8 @@ class FFloatChannelBufferedCurveModel : public IBufferedCurveModel
 public:
 	/** Create a copy of the float channel while keeping the reference to the section */
 	FFloatChannelBufferedCurveModel(const FMovieSceneFloatChannel* InMovieSceneFloatChannel, TWeakObjectPtr<UMovieSceneSection> InWeakSection,
-		TArray<FKeyPosition>&& InKeyPositions, TArray<FKeyAttributes>&& InKeyAttributes, const FString& InIntentionName, const double InValueMin, const double InValueMax)
-		: IBufferedCurveModel(MoveTemp(InKeyPositions), MoveTemp(InKeyAttributes), InIntentionName, InValueMin, InValueMax)
+		TArray<FKeyPosition>&& InKeyPositions, TArray<FKeyAttributes>&& InKeyAttributes, const FString& InLongDisplayName, const double InValueMin, const double InValueMax)
+		: IBufferedCurveModel(MoveTemp(InKeyPositions), MoveTemp(InKeyAttributes), InLongDisplayName, InValueMin, InValueMax)
 		, Channel(*InMovieSceneFloatChannel)
 		, WeakSection(InWeakSection)
 	{}
@@ -77,7 +77,7 @@ TUniquePtr<IBufferedCurveModel> FFloatChannelCurveModel::CreateBufferedCurveCopy
 		double ValueMin = 0.f, ValueMax = 1.f;
 		GetValueRange(ValueMin, ValueMax);
 
-		return MakeUnique<FFloatChannelBufferedCurveModel>(Channel, Cast<UMovieSceneSection>(GetOwningObject()), MoveTemp(KeyPositions), MoveTemp(KeyAttributes), GetIntentionName(), ValueMin, ValueMax);
+		return MakeUnique<FFloatChannelBufferedCurveModel>(Channel, Cast<UMovieSceneSection>(GetOwningObject()), MoveTemp(KeyPositions), MoveTemp(KeyAttributes), GetLongDisplayName().ToString(), ValueMin, ValueMax);
 	}
 	return nullptr;
 }
diff --git a/Engine/Source/Editor/MovieSceneTools/Private/Channels/IntegerChannelCurveModel.cpp b/Engine/Source/Editor/MovieSceneTools/Private/Channels/IntegerChannelCurveModel.cpp
--- a/Engine/Source/Editor/MovieSceneTools/Private/Channels/IntegerChannelCurveModel.cpp
+++ b/Engine/Source/Editor/MovieSceneTools/Private/Channels/IntegerChannelCurveModel.cpp
@@ -26,8 +26,8 @@ class FIntegerChannelBufferedCurveModel : public IBufferedCurveModel
 public:
 	/** Create a copy of the float channel while keeping the reference to the section */
 	FIntegerChannelBufferedCurveModel(const FMovieSceneIntegerChannel* InMovieSceneIntegerChannel, TWeakObjectPtr<UMovieSceneSection> InWeakSection,
-		TArray<FKeyPosition>&& InKeyPositions, TArray<FKeyAttributes>&& InKeyAttributes, const FString& InIntentionName, const double InValueMin, const double InValueMax)
-		: IBufferedCurveModel(MoveTemp(InKeyPositions), MoveTemp(InKeyAttributes), InIntentionName, InValueMin, InValueMax)
+		TArray<FKeyPosition>&& InKeyPositions, TArray<FKeyAttributes>&& InKeyAttributes, const FString& InLongDisplayName, const double InValueMin, const double InValueMax)
+		: IBufferedCurveModel(MoveTemp(InKeyPositions), MoveTemp(InKeyAttributes), InLongDisplayName, InValueMin, InValueMax)
 		, Channel(*InMovieSceneIntegerChannel)
 		, WeakSection(InWeakSection)
 	{}
@@ -102,7 +102,7 @@ TUniquePtr<IBufferedCurveModel> FIntegerChannelCurveModel::CreateBufferedCurveCo
 		double ValueMin = 0.f, ValueMax = 1.f;
 		GetValueRange(ValueMin, ValueMax);
 
-		return MakeUnique<FIntegerChannelBufferedCurveModel>(Channel, Cast<UMovieSceneSection>(GetOwningObject()), MoveTemp(KeyPositions), MoveTemp(KeyAttributes), GetIntentionName(), ValueMin, ValueMax);
+		return MakeUnique<FIntegerChannelBufferedCurveModel>(Channel, Cast<UMovieSceneSection>(GetOwningObject()), MoveTemp(KeyPositions), MoveTemp(KeyAttributes), GetLongDisplayName().ToString(), ValueMin, ValueMax);
 	}
 	return nullptr;
 }
diff --git a/Engine/Source/Editor/MovieSceneTools/Public/Channels/BezierChannelCurveModel.h b/Engine/Source/Editor/MovieSceneTools/Public/Channels/BezierChannelCurveModel.h
--- a/Engine/Source/Editor/MovieSceneTools/Public/Channels/BezierChannelCurveModel.h
+++ b/Engine/Source/Editor/MovieSceneTools/Public/Channels/BezierChannelCurveModel.h
@@ -22,7 +22,7 @@ class FBezierChannelBufferedCurveModel : public IBufferedCurveModel
 public:
 	/** Create a copy of the channel while keeping the reference to the section */
 	FBezierChannelBufferedCurveModel(const ChannelType* InChannel, TWeakObjectPtr<UMovieSceneSection> InWeakSection,
-		TArray<FKeyPosition>&& InKeyPositions, TArray<FKeyAttributes>&& InKeyAttributes, const FString& InIntentionName, const double InValueMin, const double InValueMax);
+		TArray<FKeyPosition>&& InKeyPositions, TArray<FKeyAttributes>&& InKeyAttributes, const FString& InLongDisplayName, const double InValueMin, const double InValueMax);
 
 	virtual void DrawCurve(const FCurveEditor& InCurveEditor, const FCurveEditorScreenSpace& InScreenSpace, TArray<TTuple<double, double>>& OutInterpolatingPoints) const override;
 
