Adding new setting to disable auto-expansion of the sequencer outliner.
Changed the curve editor filter states to a bitmask so we can communicate desired expansion state as well.

diff --git a/Engine/Source/Editor/CurveEditor/Private/Tree/CurveEditorTree.cpp b/Engine/Source/Editor/CurveEditor/Private/Tree/CurveEditorTree.cpp
--- a/Engine/Source/Editor/CurveEditor/Private/Tree/CurveEditorTree.cpp
+++ b/Engine/Source/Editor/CurveEditor/Private/Tree/CurveEditorTree.cpp
@@ -251,9 +251,9 @@ void FCurveEditorTree::ToggleExpansionState(bool bRecursive)
 	ToggleExpansionStateDelegate.Broadcast(bRecursive);
 }
 
-bool FCurveEditorTree::PerformFilterPass(TArrayView<const FCurveEditorTreeFilter* const> FilterPtrs, TArrayView<const FCurveEditorTreeItemID> ItemsToFilter, ECurveEditorTreeFilterState InheritedState)
+ECurveEditorTreeFilterState FCurveEditorTree::PerformFilterPass(TArrayView<const FCurveEditorTreeFilter* const> FilterPtrs, TArrayView<const FCurveEditorTreeItemID> ItemsToFilter, ECurveEditorTreeFilterState InheritedState)
 {
-	bool bAnyMatched = false;
+	ECurveEditorTreeFilterState ReturnedParentState = ECurveEditorTreeFilterState::NoMatch;
 
 	for (FCurveEditorTreeItemID ItemID : ItemsToFilter)
 	{
@@ -265,35 +265,45 @@ bool FCurveEditorTree::PerformFilterPass(TArrayView<const FCurveEditorTreeFilter
 
 		const FCurveEditorTreeItem& TreeItem = GetItem(ItemID);
 
-		ECurveEditorTreeFilterState FilterState         = InheritedState;
+		ECurveEditorTreeFilterState FilterState = InheritedState;
 		ECurveEditorTreeFilterState ChildInheritedState = InheritedState;
 
 		TSharedPtr<ICurveEditorTreeItem> TreeItemImpl = TreeItem.GetItem();
 		if (TreeItemImpl)
 		{
-			const bool bMatchesFilter = Algo::AnyOf(FilterPtrs, [TreeItemImpl](const FCurveEditorTreeFilter* Filter){ return TreeItemImpl->PassesFilter(Filter); });
-			if (bMatchesFilter)
+			for (const FCurveEditorTreeFilter* Filter : FilterPtrs)
 			{
-				bAnyMatched = true;
-				FilterState = ECurveEditorTreeFilterState::Match;
-				ChildInheritedState = ECurveEditorTreeFilterState::ImplicitChild;
+				const bool bMatchesFilter = TreeItemImpl->PassesFilter(Filter);
+				if (bMatchesFilter)
+				{
+					ReturnedParentState |= ECurveEditorTreeFilterState::ImplicitParent;
+					FilterState |= ECurveEditorTreeFilterState::Match;
+					ChildInheritedState |= ECurveEditorTreeFilterState::ImplicitChild;
+
+					if (Filter->ShouldExpandOnMatch())
+					{
+						ReturnedParentState |= ECurveEditorTreeFilterState::Expand;
+
+						// Once an item is matched AND set to expand there's no other flags that could be 
+						// discernibly applied, therefore we can early out and skip needlessly testing other filters
+						break;
+					}
+					// Else, we have to keep iterating incase one of the other filters is set to expand
+				}
 			}
 		}
 
 		// Run the filter on all child nodes
-		const bool bMatchedChildren = PerformFilterPass(FilterPtrs, TreeItem.GetChildren(), ChildInheritedState);
+		ECurveEditorTreeFilterState AppliedStateFromChild = PerformFilterPass(FilterPtrs, TreeItem.GetChildren(), ChildInheritedState);
+		ReturnedParentState |= AppliedStateFromChild;
 
-		// If we matched children we become an implicit parent if not already matched
-		if (bMatchedChildren && FilterState != ECurveEditorTreeFilterState::Match)
-		{
-			bAnyMatched = true;
-			FilterState = ECurveEditorTreeFilterState::ImplicitParent;
-		}
+		// If we matched children we become an implicit parent
+		FilterState |= AppliedStateFromChild;
 
 		FilterStates.SetFilterState(ItemID, FilterState);
 	}
 
-	return bAnyMatched;
+	return ReturnedParentState;
 }
 
 void FCurveEditorTree::RunFilters()
@@ -473,7 +483,7 @@ void FCurveEditorTree::SetDirectSelection(TArray<FCurveEditorTreeItemID>&& TreeI
 
 		for (FCurveEditorTreeItemID ChildID : TreeItem->GetChildren())
 		{
-			if (FilterStates.Get(ChildID) != ECurveEditorTreeFilterState::NoMatch)
+			if ((FilterStates.Get(ChildID) & ECurveEditorTreeFilterState::MatchBitMask) != ECurveEditorTreeFilterState::NoMatch)
 			{
 				TreeItems.Add(ChildID);
 			}
diff --git a/Engine/Source/Editor/CurveEditor/Private/Tree/SCurveEditorTree.cpp b/Engine/Source/Editor/CurveEditor/Private/Tree/SCurveEditorTree.cpp
--- a/Engine/Source/Editor/CurveEditor/Private/Tree/SCurveEditorTree.cpp
+++ b/Engine/Source/Editor/CurveEditor/Private/Tree/SCurveEditorTree.cpp
@@ -29,7 +29,7 @@ struct SCurveEditorTableRow : SMultiColumnTableRow<FCurveEditorTreeItemID>
 	{
 		TSharedPtr<FCurveEditor> CurveEditor = WeakCurveEditor.Pin();
 
-		const bool bIsMatch = CurveEditor.IsValid() && ( CurveEditor->GetTree()->GetFilterState(TreeItemID) == ECurveEditorTreeFilterState::Match );
+		const bool bIsMatch = CurveEditor.IsValid() && ((CurveEditor->GetTree()->GetFilterState(TreeItemID) & ECurveEditorTreeFilterState::Match) != ECurveEditorTreeFilterState::NoMatch);
 		return bIsMatch ? GetForegroundBasedOnSelection() : FSlateColor::UseSubduedForeground();
 	}
 
@@ -109,6 +109,7 @@ void SCurveEditorTree::Construct(const FArguments& InArgs, TSharedPtr<FCurveEdit
 				this->OnTreeSelectionChanged(InItemID, Type);
 			}
 		)
+		.AllowInvisibleItemSelection(true)
 	);
 
 	CurveEditor->GetTree()->Events.OnItemsChanged.AddSP(this, &SCurveEditorTree::RefreshTree);
@@ -146,7 +147,11 @@ void SCurveEditorTree::RefreshTree()
 
 			while (ParentItem)
 			{
-				PreFilterExpandedItems.Add(ParentItem->GetID());
+				const FCurveEditorTreeItemID ParentID = ParentItem->GetID();
+				if (IsItemExpanded(ParentID))
+				{
+					PreFilterExpandedItems.Add(ParentID);
+				}
 				ParentItem = CurveEditorTree->FindItem(ParentItem->GetParentID());
 			}
 		}
@@ -163,7 +168,7 @@ void SCurveEditorTree::RefreshTree()
 	// Repopulate root tree items based on filters
 	for (FCurveEditorTreeItemID RootItemID : CurveEditor->GetRootTreeItems())
 	{
-		if (FilterStates.Get(RootItemID) != ECurveEditorTreeFilterState::NoMatch)
+		if ((FilterStates.Get(RootItemID) & ECurveEditorTreeFilterState::MatchBitMask) != ECurveEditorTreeFilterState::NoMatch)
 		{
 			RootItems.Add(RootItemID);
 		}
@@ -174,16 +179,24 @@ void SCurveEditorTree::RefreshTree()
 
 	if (FilterStates.IsActive())
 	{
-		// If a filter is active, all matched items and their parents are expanded
-		ClearExpandedItems();
-		for (const TTuple<FCurveEditorTreeItemID, FCurveEditorTreeItem>& Pair : CurveEditorTree->GetAllItems())
-		{
-			ECurveEditorTreeFilterState FilterState = FilterStates.Get(Pair.Key);
+		TArray<FCurveEditorTreeItemID> ExpandedItems;
+		ExpandedItems.Reserve(FilterStates.GetNumMatchedImplicitly());
+
+		FilterStates.ForEachItemState([&ExpandedItems](const TTuple<FCurveEditorTreeItemID, ECurveEditorTreeFilterState>& FilterState)
+			{
+				if ((FilterState.Value & ECurveEditorTreeFilterState::Expand) != ECurveEditorTreeFilterState::NoMatch)
+				{
+					ExpandedItems.Add(FilterState.Key);
+				}
+			}
+		);
 
-			// Expand any matched items or parents of matched items
-			if (FilterState == ECurveEditorTreeFilterState::Match || FilterState == ECurveEditorTreeFilterState::ImplicitParent)
+		if (ExpandedItems.Num() > 0)
+		{
+			ClearExpandedItems();
+			for (const FCurveEditorTreeItemID& ItemID : ExpandedItems)
 			{
-				SetItemExpansion(Pair.Key, true);
+				SetItemExpansion(ItemID, true);
 			}
 		}
 	}
@@ -274,7 +287,7 @@ void SCurveEditorTree::GetTreeItemChildren(FCurveEditorTreeItemID Parent, TArray
 
 	for (FCurveEditorTreeItemID ChildID : CurveEditor->GetTreeItem(Parent).GetChildren())
 	{
-		if (FilterStates.Get(ChildID) != ECurveEditorTreeFilterState::NoMatch)
+		if ((FilterStates.Get(ChildID) & ECurveEditorTreeFilterState::MatchBitMask) != ECurveEditorTreeFilterState::NoMatch)
 		{
 			OutChildren.Add(ChildID);
 		}
diff --git a/Engine/Source/Editor/CurveEditor/Public/Tree/CurveEditorTree.h b/Engine/Source/Editor/CurveEditor/Public/Tree/CurveEditorTree.h
--- a/Engine/Source/Editor/CurveEditor/Public/Tree/CurveEditorTree.h
+++ b/Engine/Source/Editor/CurveEditor/Public/Tree/CurveEditorTree.h
@@ -9,6 +9,8 @@
 #include "Templates/SharedPointer.h"
 #include "Delegates/Delegate.h"
 #include "CurveEditorTypes.h"
+#include "Misc/EnumClassFlags.h"
+#include "Algo/ForEach.h"
 
 enum class ECurveEditorTreeFilterType : uint32;
 
@@ -18,21 +20,27 @@ struct FCurveEditorTreeFilter;
 class FCurveEditor;
 class FCurveEditorTree;
 
-/** Enumeration specifying how a specific tree item has matched the current set of filters */
+/** Enumeration of bitmask values specifying how a specific tree item is interpreted by the current set of filters */
 enum class ECurveEditorTreeFilterState : uint8
 {
 	/** The item did not match any filter, and neither did any of its parents or children */
-	NoMatch,
-
-	/** Neither this item nor any of its children match filters, but one of its parents did (ie it resides within a matched item) */
-	ImplicitChild,
+	NoMatch        = 0x00,
+
+	/** One of this item's parents matched a filter (ie it resides within a matched item) */
+	ImplicitChild  = (1<<0),
+	
+	/** One of this item's descendant children matched a filter (ie it is a parent of a matched item) */
+	ImplicitParent = (1<<1),
+	
+	/** This item itself matched one or more of the filters */
+	Match          = (1<<2),
 
-	/** Neither this item nor any of its parents match the filters, but one of its descendant children did (ie it is a parent of a matched item) */
-	ImplicitParent,
+	/** This item in the tree should be expanded according to one or more of the filters */
+	Expand         = (1<<3),
 
-	/** This item itself matched one or more of the filters */
-	Match,
+	MatchBitMask   = (ImplicitParent | Match | ImplicitChild),
 };
+ENUM_CLASS_FLAGS(ECurveEditorTreeFilterState);
 
 /**
  * Scoped guard that prevents the broadcast of tree events for the duration of its lifetime. Will trigger necessary events after the last remaining guard has been destroyed.
@@ -250,6 +258,12 @@ struct FCurveEditorFilterStates
 		return State ? *State : ECurveEditorTreeFilterState::NoMatch;
 	}
 
+	template <typename CallableT>
+	void ForEachItemState(CallableT Callable) const
+	{
+		Algo::ForEach(FilterStates, Callable);
+	}
+
 	/**
 	 * Assign a new filter state to an item
 	 */
@@ -258,17 +272,17 @@ struct FCurveEditorFilterStates
 		const ECurveEditorTreeFilterState* Existing = FilterStates.Find(ItemID);
 		if (Existing)
 		{
-			if (*Existing  == ECurveEditorTreeFilterState::Match)
+			if ((*Existing & ECurveEditorTreeFilterState::Match) != ECurveEditorTreeFilterState::NoMatch)
 			{
 				--NumMatched;
 			}
-			else if (*Existing != ECurveEditorTreeFilterState::NoMatch)
+			else if ((*Existing & ECurveEditorTreeFilterState::MatchBitMask) != ECurveEditorTreeFilterState::NoMatch)
 			{
 				--NumMatchedImplicitly;
 			}
 		}
 
-		if (NewState == ECurveEditorTreeFilterState::NoMatch)
+		if ((NewState & ECurveEditorTreeFilterState::MatchBitMask) == ECurveEditorTreeFilterState::NoMatch)
 		{
 			FilterStates.Remove(ItemID);
 		}
@@ -276,7 +290,7 @@ struct FCurveEditorFilterStates
 		{
 			FilterStates.Add(ItemID, NewState);
 
-			if (NewState == ECurveEditorTreeFilterState::Match)
+			if ((NewState & ECurveEditorTreeFilterState::Match) != ECurveEditorTreeFilterState::NoMatch)
 			{
 				++NumMatched;
 			}
@@ -508,9 +522,9 @@ private:
 	 * @param FilterPtrs     Array of non-null pointers to filters to use. Items are considered matched if they match any filter in this array.
 	 * @param ItemsToFilter  Array item IDs to filter
 	 * @param InheritedState The filter state for each item to receive if it does not directly match a filter (either ECurveEditorTreeFilterState::NoMatch or ECurveEditorTreeFilterState::InheritedChild)
-	 * @return Whether any of the items or any their recursive children matched any filter
+	 * @return Raised state flags which should be applied to parents of the filtered items.
 	 */
-	bool PerformFilterPass(TArrayView<const FCurveEditorTreeFilter* const> FilterPtrs, TArrayView<const FCurveEditorTreeItemID> ItemsToFilter, ECurveEditorTreeFilterState InheritedState);
+	ECurveEditorTreeFilterState PerformFilterPass(TArrayView<const FCurveEditorTreeFilter* const> FilterPtrs, TArrayView<const FCurveEditorTreeItemID> ItemsToFilter, ECurveEditorTreeFilterState InheritedState);
 
 	/** 
 	 * Recursively sorts the tree item ids using the sort predicate.
diff --git a/Engine/Source/Editor/CurveEditor/Public/Tree/CurveEditorTreeFilter.h b/Engine/Source/Editor/CurveEditor/Public/Tree/CurveEditorTreeFilter.h
--- a/Engine/Source/Editor/CurveEditor/Public/Tree/CurveEditorTreeFilter.h
+++ b/Engine/Source/Editor/CurveEditor/Public/Tree/CurveEditorTreeFilter.h
@@ -35,6 +35,7 @@ struct CURVEEDITOR_API FCurveEditorTreeFilter
 	FCurveEditorTreeFilter(ECurveEditorTreeFilterType InFilterType, int32 InFilterPass)
 		: FilterType(InFilterType)
 		, FilterPass(InFilterPass)
+		, bExpandToMatchedItems(true)
 	{}
 
 	virtual ~FCurveEditorTreeFilter() {}
@@ -55,6 +56,14 @@ struct CURVEEDITOR_API FCurveEditorTreeFilter
 		return FilterPass;
 	}
 
+	/**
+	 *  @return If tree paths should be expanded down to matched items
+	 */
+	bool ShouldExpandOnMatch() const
+	{
+		return bExpandToMatchedItems;
+	}
+
 public:
 
 	/**
@@ -72,6 +81,9 @@ protected:
 
 	/** Defines which pass this filter should be applied in */
 	int32 FilterPass;
+
+	/** Determines if tree paths should be expanded down to matched items */
+	bool bExpandToMatchedItems : 1;
 };
 
 /** A specific text token (containing neither spaces nor .) */
diff --git a/Engine/Source/Editor/Sequencer/Private/SSequencer.cpp b/Engine/Source/Editor/Sequencer/Private/SSequencer.cpp
--- a/Engine/Source/Editor/Sequencer/Private/SSequencer.cpp
+++ b/Engine/Source/Editor/Sequencer/Private/SSequencer.cpp
@@ -1311,15 +1311,16 @@ void SSequencer::HandleOutlinerNodeSelectionChanged()
 	TSharedPtr<FCurveEditor> CurveEditor = Sequencer->GetCurveEditor();
 	if (CurveEditor && CurveEditorTree)
 	{
+		const USequencerSettings* SequencerSettings = GetSequencerSettings();
 		// If we're isolating to the selection and there is one, add the filter
-		if (GetSequencerSettings()->ShouldIsolateToCurveEditorSelection() && SelectedDisplayNodes.Num() != 0)
+		if (SequencerSettings->ShouldIsolateToCurveEditorSelection() && SelectedDisplayNodes.Num() != 0)
 		{
 			if (!SequencerSelectionCurveEditorFilter)
 			{
 				SequencerSelectionCurveEditorFilter = MakeShared<FSequencerSelectionCurveFilter>();
 			}
 
-			SequencerSelectionCurveEditorFilter->Update(Sequencer->GetSelection().GetSelectedOutlinerNodes());
+			SequencerSelectionCurveEditorFilter->Update(Sequencer->GetSelection().GetSelectedOutlinerNodes(), SequencerSettings->GetAutoExpandTreeView());
 
 			CurveEditor->GetTree()->AddFilter(SequencerSelectionCurveEditorFilter);
 		}
diff --git a/Engine/Source/Editor/Sequencer/Private/Sequencer.cpp b/Engine/Source/Editor/Sequencer/Private/Sequencer.cpp
--- a/Engine/Source/Editor/Sequencer/Private/Sequencer.cpp
+++ b/Engine/Source/Editor/Sequencer/Private/Sequencer.cpp
@@ -7117,6 +7117,18 @@ void FSequencer::SelectByNthCategoryNode(UMovieSceneSection* Section, int Index,
 	}
 	if (bSelect)
 	{
+		if (Settings->GetAutoExpandTreeView())
+		{
+			for (const TSharedRef<FSequencerDisplayNode>& DisplayNode : NodesToSelect)
+			{
+				if (DisplayNode->GetParent().IsValid() && DisplayNode->GetParent()->GetType() == ESequencerNode::Track && !DisplayNode->GetParent()->IsExpanded())
+				{
+					DisplayNode->GetParent()->SetExpansionState(true);
+					break;
+				}
+			}
+		}
+
 		if (NodesToSelect.Num() > 0)
 		{
 			SequencerWidget->GetTreeView()->RequestScrollIntoView(NodesToSelect[0]);
@@ -7168,7 +7180,7 @@ void FSequencer::SelectByChannels(UMovieSceneSection* Section, TArrayView<const
 	}
 	
 	if (bSelect)
-	{
+	{	
 		if (NodesToSelect.Num() > 0)
 		{
 			//todo hide behind preference 
@@ -7232,9 +7244,12 @@ void FSequencer::SelectByChannels(UMovieSceneSection* Section, const TArray<FNam
 	{
 		for (const TSharedRef<FSequencerDisplayNode>& DisplayNode : Nodes)
 		{
-			if (DisplayNode->GetParent().IsValid() && DisplayNode->GetParent()->GetType() == ESequencerNode::Track && !DisplayNode->GetParent()->IsExpanded())
+			if (Settings->GetAutoExpandTreeView())
 			{
-				DisplayNode->GetParent()->SetExpansionState(true);
+				if (DisplayNode->GetParent().IsValid() && DisplayNode->GetParent()->GetType() == ESequencerNode::Track && !DisplayNode->GetParent()->IsExpanded())
+				{
+					DisplayNode->GetParent()->SetExpansionState(true);
+				}
 			}
 			NodesToSelect.Add(DisplayNode);
 		}
diff --git a/Engine/Source/Editor/Sequencer/Private/SequencerSelectionCurveFilter.h b/Engine/Source/Editor/Sequencer/Private/SequencerSelectionCurveFilter.h
--- a/Engine/Source/Editor/Sequencer/Private/SequencerSelectionCurveFilter.h
+++ b/Engine/Source/Editor/Sequencer/Private/SequencerSelectionCurveFilter.h
@@ -23,7 +23,7 @@ struct FSequencerSelectionCurveFilter : FCurveEditorTreeFilter
 	/**
 	 * Adds all selected nodes and their object parents to the NodesToFilter set
 	 */
-	void Update(const TSet<TSharedRef<FSequencerDisplayNode>>& SelectedNodes)
+	void Update(const TSet<TSharedRef<FSequencerDisplayNode>>& SelectedNodes, const bool bExpandTreeToSelectedNodes = true)
 	{
 		NodesToFilter.Empty(SelectedNodes.Num());
 
@@ -43,6 +43,8 @@ struct FSequencerSelectionCurveFilter : FCurveEditorTreeFilter
 				Parent = Parent->GetParent();
 			}
 		}
+
+		bExpandToMatchedItems = bExpandTreeToSelectedNodes;
 	}
 
 	bool Match(TSharedRef<const FSequencerDisplayNode> InNode) const
diff --git a/Engine/Source/Editor/Sequencer/Private/SequencerSettings.cpp b/Engine/Source/Editor/Sequencer/Private/SequencerSettings.cpp
--- a/Engine/Source/Editor/Sequencer/Private/SequencerSettings.cpp
+++ b/Engine/Source/Editor/Sequencer/Private/SequencerSettings.cpp
@@ -57,6 +57,7 @@ USequencerSettings::USequencerSettings( const FObjectInitializer& ObjectInitiali
 	bVisualizePreAndPostRoll = true;
 	TrajectoryPathCap = 250;
 	FrameNumberDisplayFormat = EFrameNumberDisplayFormats::Seconds;
+	bAutoExpandOutlinerTreeOnChildSelection = true;
 }
 
 void USequencerSettings::PostEditChangeProperty(FPropertyChangedEvent& PropertyChangedEvent)
@@ -832,3 +833,12 @@ void USequencerSettings::SetMovieRendererName(const FString& InMovieRendererName
 		SaveConfig();
 	}
 }
+
+void USequencerSettings::SetAutoExpandTreeView(bool bInAutoExpandOutlinerTree)
+{
+	if (bInAutoExpandOutlinerTree != bAutoExpandOutlinerTreeOnChildSelection)
+	{
+		bAutoExpandOutlinerTreeOnChildSelection = bInAutoExpandOutlinerTree;
+		SaveConfig();
+	}
+}
\ No newline at end of file
diff --git a/Engine/Source/Editor/Sequencer/Public/SequencerSettings.h b/Engine/Source/Editor/Sequencer/Public/SequencerSettings.h
--- a/Engine/Source/Editor/Sequencer/Public/SequencerSettings.h
+++ b/Engine/Source/Editor/Sequencer/Public/SequencerSettings.h
@@ -372,6 +372,12 @@ public:
 	/** Sets the movie renderer to use */
 	void SetMovieRendererName(const FString& InMovieRendererName);
 
+	
+	/** Gets whether or not to expand the outliner tree view when a child element is selected (from outside of the tree view). */
+	bool GetAutoExpandTreeView() const { return bAutoExpandOutlinerTreeOnChildSelection; }
+	/** Sets whether or not to expand the outliner tree view when a child element is selected (from outside of the tree view). */
+	void SetAutoExpandTreeView(bool bInAutoExpandOutlinerTree);
+
 protected:
 
 	/** The auto change mode (auto-key, auto-track or none). */
@@ -584,6 +590,10 @@ protected:
 	UPROPERTY(config, EditAnywhere, Category=General)
 	FString MovieRendererName;
 
+	/** Whether to expand the sequencer tree view when a child element is selected (from outside of the tree view). */
+	UPROPERTY(config, EditAnywhere, Category= General)
+	bool bAutoExpandOutlinerTreeOnChildSelection;
+
 	FOnEvaluateSubSequencesInIsolationChanged OnEvaluateSubSequencesInIsolationChangedEvent;
 	FOnShowSelectedNodesOnlyChanged OnShowSelectedNodesOnlyChangedEvent;
 	FOnAllowEditsModeChanged OnAllowEditsModeChangedEvent;
