Delaying CurveEditor command registration until PostEngineInit

diff --git a/Engine/Source/Editor/CurveEditor/CurveEditor.Build.cs b/Engine/Source/Editor/CurveEditor/CurveEditor.Build.cs
--- a/Engine/Source/Editor/CurveEditor/CurveEditor.Build.cs
+++ b/Engine/Source/Editor/CurveEditor/CurveEditor.Build.cs
@@ -20,6 +20,7 @@ public class CurveEditor : ModuleRules
 				"TimeManagement",
 				"UnrealEd",
 				"SequencerWidgets",
+				"ToolMenus",
 			}
 		);
 
diff --git a/Engine/Source/Editor/CurveEditor/Private/CurveEditorModule.cpp b/Engine/Source/Editor/CurveEditor/Private/CurveEditorModule.cpp
--- a/Engine/Source/Editor/CurveEditor/Private/CurveEditorModule.cpp
+++ b/Engine/Source/Editor/CurveEditor/Private/CurveEditorModule.cpp
@@ -6,6 +6,7 @@
 #include "CurveEditor.h"
 #include "CurveEditorViewRegistry.h"
 #include "Framework/MultiBox/MultiBoxExtender.h"
+#include "ToolMenus.h"
 
 class FCurveEditorModule : public ICurveEditorModule
 {
@@ -14,7 +15,14 @@ public:
 	{
 		if (GIsEditor)
 		{
-			FCurveEditorCommands::Register();
+			if (UToolMenus::TryGet())
+			{
+				FCurveEditorCommands::Register();
+			}
+			else
+			{
+				FCoreDelegates::OnPostEngineInit.AddStatic(&FCurveEditorCommands::Register);
+			}
 		}
 	}
 
