Curve Editor: Get all the points on the curve, not just the drawn points

diff --git a/Engine/Source/Editor/CurveEditor/Private/Views/SInteractiveCurveEditorView.cpp b/Engine/Source/Editor/CurveEditor/Private/Views/SInteractiveCurveEditorView.cpp
--- a/Engine/Source/Editor/CurveEditor/Private/Views/SInteractiveCurveEditorView.cpp
+++ b/Engine/Source/Editor/CurveEditor/Private/Views/SInteractiveCurveEditorView.cpp
@@ -546,7 +546,7 @@ bool SInteractiveCurveEditorView::GetCurveWithinWidgetRange(const FSlateRect& Wi
 
 	// Iterate through all of our interpolating points and terminates if one overlaps the marquee. Both of these coordinate systems
 	// are in screen space pixels.
-	bool bFound = false;
+	TSet<FCurveModelID> CurveIDs;
 	for (const FCurveDrawParams& DrawParams : CachedDrawParams)
 	{
 		for (int32 InterpolatingPointIndex = 1; InterpolatingPointIndex < DrawParams.InterpolatingPoints.Num(); InterpolatingPointIndex++)
@@ -559,20 +559,28 @@ bool SInteractiveCurveEditorView::GetCurveWithinWidgetRange(const FSlateRect& Wi
 
 			if (FMath::LineBoxIntersection(WidgetRectangleBox, Start, End, StartToEnd))
 			{
-				for (int32 PointIndex = 0; PointIndex < DrawParams.Points.Num(); PointIndex++)
-				{
-					const FCurvePointInfo& Point = DrawParams.Points[PointIndex];
+				CurveIDs.Add(DrawParams.GetID());
+			}
+		}
+	}
 
-					OutPoints->Add(FCurvePointHandle(DrawParams.GetID(), Point.Type, Point.KeyHandle));
-				}
+	bool bPointsAdded = false;
+	for (const FCurveModelID& CurveID : CurveIDs)
+	{
+		if (const FCurveModel* Curve = CurveEditor->FindCurve(CurveID))
+		{
+			TArray<FKeyHandle> KeyHandles;
+			Curve->GetKeys(*CurveEditor, TNumericLimits<double>::Lowest(), TNumericLimits<double>::Max(), TNumericLimits<double>::Lowest(), TNumericLimits<double>::Max(), KeyHandles);
 
-				bFound = true;
-				break;
+			for (const FKeyHandle& KeyHandle : KeyHandles)
+			{
+				OutPoints->Add(FCurvePointHandle(CurveID, ECurvePointType::Key, KeyHandle));
+				bPointsAdded = true;
 			}
 		}
 	}
 
-	return false;
+	return bPointsAdded;
 }
 
 void SInteractiveCurveEditorView::UpdateCurveProximities(FVector2D MousePixel)
