Curve Editor: Add select all keys forward and backwards shortcuts: Ctrl-] and Ctrl-[

diff --git a/Engine/Source/Editor/CurveEditor/Private/CurveEditor.cpp b/Engine/Source/Editor/CurveEditor/Private/CurveEditor.cpp
--- a/Engine/Source/Editor/CurveEditor/Private/CurveEditor.cpp
+++ b/Engine/Source/Editor/CurveEditor/Private/CurveEditor.cpp
@@ -369,6 +369,8 @@ void FCurveEditor::BindCommands()
 	CommandList->MapAction(FCurveEditorCommands::Get().ClearSelectionRange, FExecuteAction::CreateSP(this, &FCurveEditor::ClearSelectionRange));
 
 	CommandList->MapAction(FCurveEditorCommands::Get().SelectAllKeys, FExecuteAction::CreateSP(this, &FCurveEditor::SelectAllKeys));
+	CommandList->MapAction(FCurveEditorCommands::Get().SelectForward, FExecuteAction::CreateSP(this, &FCurveEditor::SelectForward));
+	CommandList->MapAction(FCurveEditorCommands::Get().SelectBackward, FExecuteAction::CreateSP(this, &FCurveEditor::SelectBackward));
 
 	{
 		FExecuteAction   ToggleInputSnapping     = FExecuteAction::CreateSP(this,   &FCurveEditor::ToggleInputSnapping);
@@ -1020,7 +1022,7 @@ void FCurveEditor::ClearSelectionRange()
 }
 
 void FCurveEditor::SelectAllKeys()
-{		
+{
 	for (FCurveModelID ID : GetEditedCurves())
 	{
 		if (FCurveModel* Curve = FindCurve(ID))
@@ -1032,6 +1034,56 @@ void FCurveEditor::SelectAllKeys()
 	}
 }
 
+void FCurveEditor::SelectForward()
+{
+	Selection.Clear();
+
+	TSharedPtr<ITimeSliderController> TimeSliderController = WeakTimeSliderController.Pin();
+	if (!TimeSliderController.IsValid())
+	{
+		return;
+	}
+
+	FFrameRate TickResolution = TimeSliderController->GetTickResolution();
+
+	double CurrentTime = TickResolution.AsSeconds(TimeSliderController->GetScrubPosition());
+
+	for (FCurveModelID ID : GetEditedCurves())
+	{
+		if (FCurveModel* Curve = FindCurve(ID))
+		{
+			TArray<FKeyHandle> KeyHandles;
+			Curve->GetKeys(*this, CurrentTime, TNumericLimits<double>::Max(), TNumericLimits<double>::Lowest(), TNumericLimits<double>::Max(), KeyHandles);
+			Selection.Add(ID, ECurvePointType::Key, KeyHandles);
+		}
+	}
+}
+
+void FCurveEditor::SelectBackward()
+{
+	Selection.Clear();
+
+	TSharedPtr<ITimeSliderController> TimeSliderController = WeakTimeSliderController.Pin();
+	if (!TimeSliderController.IsValid())
+	{
+		return;
+	}
+
+	FFrameRate TickResolution = TimeSliderController->GetTickResolution();
+
+	double CurrentTime = TickResolution.AsSeconds(TimeSliderController->GetScrubPosition());
+
+	for (FCurveModelID ID : GetEditedCurves())
+	{
+		if (FCurveModel* Curve = FindCurve(ID))
+		{
+			TArray<FKeyHandle> KeyHandles;
+			Curve->GetKeys(*this, TNumericLimits<double>::Min(), CurrentTime, TNumericLimits<double>::Lowest(), TNumericLimits<double>::Max(), KeyHandles);
+			Selection.Add(ID, ECurvePointType::Key, KeyHandles);
+		}
+	}
+}
+
 bool FCurveEditor::IsInputSnappingEnabled() const
 {
 	return InputSnapEnabledAttribute.Get();
diff --git a/Engine/Source/Editor/CurveEditor/Private/CurveEditorCommands.cpp b/Engine/Source/Editor/CurveEditor/Private/CurveEditorCommands.cpp
--- a/Engine/Source/Editor/CurveEditor/Private/CurveEditorCommands.cpp
+++ b/Engine/Source/Editor/CurveEditor/Private/CurveEditorCommands.cpp
@@ -78,6 +78,9 @@ void FCurveEditorCommands::RegisterCommands()
 	// Deselect any keys that the user has selected.
 	UI_COMMAND(DeselectAllKeys, "Deselect All Keys", "Clears your current key selection.", EUserInterfaceActionType::Button, FInputChord(EModifierKey::Control, EKeys::D));
 
+	UI_COMMAND(SelectForward, "Select All Keys Forward", "Select all keys forward from the current time", EUserInterfaceActionType::Button, FInputChord(EModifierKey::Control, EKeys::RightBracket));
+	UI_COMMAND(SelectBackward, "Select All Keys Backward", "Select all keys backward from the current time", EUserInterfaceActionType::Button, FInputChord(EModifierKey::Control, EKeys::LeftBracket));
+
 	// Buffer and Apply Curves. Like copy and paste, but directly onto the curves they were stored from.
 	// These names are overwritten in CurveEditorContextMenu to show the number of stashed curves.
 	UI_COMMAND(BufferVisibleCurves, "Store Curves", "Stores a copy of the selected curves which can be applied back onto themselves.", EUserInterfaceActionType::Button, FInputChord());
diff --git a/Engine/Source/Editor/CurveEditor/Public/CurveEditor.h b/Engine/Source/Editor/CurveEditor/Public/CurveEditor.h
--- a/Engine/Source/Editor/CurveEditor/Public/CurveEditor.h
+++ b/Engine/Source/Editor/CurveEditor/Public/CurveEditor.h
@@ -201,6 +201,8 @@ public:
 
 	/** Selection */
 	void SelectAllKeys();
+	void SelectForward();
+	void SelectBackward();
 
 	/** Toggle the expansion state of the selected nodes or all nodes if none selected */
 	void ToggleExpandCollapseNodes(bool bRecursive);
diff --git a/Engine/Source/Editor/CurveEditor/Public/CurveEditorCommands.h b/Engine/Source/Editor/CurveEditor/Public/CurveEditorCommands.h
--- a/Engine/Source/Editor/CurveEditor/Public/CurveEditorCommands.h
+++ b/Engine/Source/Editor/CurveEditor/Public/CurveEditorCommands.h
@@ -85,6 +85,9 @@ public:
 	TSharedPtr<FUICommandInfo> SelectAllKeys;
 	TSharedPtr<FUICommandInfo> DeselectAllKeys;
 
+	TSharedPtr<FUICommandInfo> SelectForward;
+	TSharedPtr<FUICommandInfo> SelectBackward;
+
 	TSharedPtr<FUICommandInfo> BufferVisibleCurves;
 	TSharedPtr<FUICommandInfo> SwapBufferedCurves;
 	TSharedPtr<FUICommandInfo> ApplyBufferedCurves;
