Curve Editor: Curve points now default to the same color as the curve. Selected points defaults to white and is controllable by a preference.

The default color for curves is now a darker gray instead of White since the selection color is White.
The blue channel for Z and B channels is now brighter because it was too dark on the dark background.

diff --git a/Engine/Plugins/FX/Niagara/Source/Niagara/Private/NiagaraDataInterfaceColorCurve.cpp b/Engine/Plugins/FX/Niagara/Source/Niagara/Private/NiagaraDataInterfaceColorCurve.cpp
--- a/Engine/Plugins/FX/Niagara/Source/Niagara/Private/NiagaraDataInterfaceColorCurve.cpp
+++ b/Engine/Plugins/FX/Niagara/Source/Niagara/Private/NiagaraDataInterfaceColorCurve.cpp
@@ -149,10 +149,10 @@ bool UNiagaraDataInterfaceColorCurve::Equals(const UNiagaraDataInterface* Other)
 
 void UNiagaraDataInterfaceColorCurve::GetCurveData(TArray<FCurveData>& OutCurveData)
 {
-	OutCurveData.Add(FCurveData(&RedCurve, TEXT("Red"), FLinearColor::Red));
-	OutCurveData.Add(FCurveData(&GreenCurve, TEXT("Green"), FLinearColor::Green));
-	OutCurveData.Add(FCurveData(&BlueCurve, TEXT("Blue"), FLinearColor::Blue));
-	OutCurveData.Add(FCurveData(&AlphaCurve, TEXT("Alpha"), FLinearColor::White));
+	OutCurveData.Add(FCurveData(&RedCurve, TEXT("Red"), FLinearColor(1.0f, 0.05f, 0.05f)));
+	OutCurveData.Add(FCurveData(&GreenCurve, TEXT("Green"), FLinearColor(0.05f, 1.0f, 0.05f)));
+	OutCurveData.Add(FCurveData(&BlueCurve, TEXT("Blue"), FLinearColor(0.1f, 0.2f, 1.0f)));
+	OutCurveData.Add(FCurveData(&AlphaCurve, TEXT("Alpha"), FLinearColor(0.2f, 0.2f, 0.2f)));
 }
 
 void UNiagaraDataInterfaceColorCurve::GetFunctions(TArray<FNiagaraFunctionSignature>& OutFunctions)
diff --git a/Engine/Plugins/FX/Niagara/Source/Niagara/Private/NiagaraDataInterfaceCurve.cpp b/Engine/Plugins/FX/Niagara/Source/Niagara/Private/NiagaraDataInterfaceCurve.cpp
--- a/Engine/Plugins/FX/Niagara/Source/Niagara/Private/NiagaraDataInterfaceCurve.cpp
+++ b/Engine/Plugins/FX/Niagara/Source/Niagara/Private/NiagaraDataInterfaceCurve.cpp
@@ -85,7 +85,7 @@ bool UNiagaraDataInterfaceCurve::Equals(const UNiagaraDataInterface* Other) cons
 
 void UNiagaraDataInterfaceCurve::GetCurveData(TArray<FCurveData>& OutCurveData)
 {
-	OutCurveData.Add(FCurveData(&Curve, NAME_None, FLinearColor::Red));
+	OutCurveData.Add(FCurveData(&Curve, NAME_None, FLinearColor(1.0f, 0.05f, 0.05f)));
 }
 
 void UNiagaraDataInterfaceCurve::GetFunctions(TArray<FNiagaraFunctionSignature>& OutFunctions)
diff --git a/Engine/Source/Editor/CurveEditor/Private/CurveEditorSettings.cpp b/Engine/Source/Editor/CurveEditor/Private/CurveEditorSettings.cpp
--- a/Engine/Source/Editor/CurveEditor/Private/CurveEditorSettings.cpp
+++ b/Engine/Source/Editor/CurveEditor/Private/CurveEditorSettings.cpp
@@ -13,6 +13,7 @@ UCurveEditorSettings::UCurveEditorSettings()
 	ZoomPosition = ECurveEditorZoomPosition::CurrentTime;
 	bSnapTimeToSelection = false;
 
+	SelectionColor = FLinearColor::White;
 	ParentSpaceCustomColor = FLinearColor(.93, .31, .19); //pastel orange
 	WorldSpaceCustomColor = FLinearColor(.198, .610, .558); //pastel teal
 }
@@ -129,6 +130,20 @@ void UCurveEditorSettings::SetSnapTimeToSelection(bool bInSnapTimeToSelection)
 	}
 }
 
+FLinearColor UCurveEditorSettings::GetSelectionColor() const
+{
+	return SelectionColor;
+}
+
+void UCurveEditorSettings::SetSelectionColor(const FLinearColor& InColor)
+{
+	if (SelectionColor != InColor)
+	{
+		SelectionColor = InColor;
+		SaveConfig();
+	}
+}
+
 TOptional<FLinearColor> UCurveEditorSettings::GetCustomColor(UClass* InClass, const FString& InPropertyName) const
 {
 	TOptional<FLinearColor> Color;
diff --git a/Engine/Source/Editor/CurveEditor/Private/RichCurveEditorModel.cpp b/Engine/Source/Editor/CurveEditor/Private/RichCurveEditorModel.cpp
--- a/Engine/Source/Editor/CurveEditor/Private/RichCurveEditorModel.cpp
+++ b/Engine/Source/Editor/CurveEditor/Private/RichCurveEditorModel.cpp
@@ -273,19 +273,15 @@ void FRichCurveEditorModel::GetKeyDrawInfo(ECurvePointType PointType, const FKey
 		{
 		case ERichCurveInterpMode::RCIM_Constant:
 			OutDrawInfo.Brush = FAppStyle::GetBrush("GenericCurveEditor.ConstantKey");
-			OutDrawInfo.Tint = FLinearColor(0, 0.45f, 0.70f);
 			break;
 		case ERichCurveInterpMode::RCIM_Linear:
 			OutDrawInfo.Brush = FAppStyle::GetBrush("GenericCurveEditor.LinearKey");
-			OutDrawInfo.Tint = FLinearColor(0, 0.62f, 0.46f);
 			break;
 		case ERichCurveInterpMode::RCIM_Cubic:
 			OutDrawInfo.Brush = FAppStyle::GetBrush("GenericCurveEditor.CubicKey");
-			OutDrawInfo.Tint = FLinearColor::White;
 			break;
 		default:
 			OutDrawInfo.Brush = FAppStyle::GetBrush("GenericCurveEditor.Key");
-			OutDrawInfo.Tint = FLinearColor::White;
 			break;
 		}
 	}
diff --git a/Engine/Source/Editor/CurveEditor/Private/SCurveViewerPanel.cpp b/Engine/Source/Editor/CurveEditor/Private/SCurveViewerPanel.cpp
--- a/Engine/Source/Editor/CurveEditor/Private/SCurveViewerPanel.cpp
+++ b/Engine/Source/Editor/CurveEditor/Private/SCurveViewerPanel.cpp
@@ -70,13 +70,20 @@ void SCurveViewerPanel::DrawCurves(const FGeometry& AllottedGeometry, FSlateWind
 				const FKeyDrawInfo& PointDrawInfo = Params.GetKeyDrawInfo(Point.Type, PointIndex);
 
 				const int32 KeyLayerId = BaseLayerId + Point.LayerBias + CurveViewConstants::ELayerOffset::Keys;
+				FLinearColor PointTint = PointDrawInfo.Tint.IsSet() ? PointDrawInfo.Tint.GetValue() : Params.Color;
+
+				// Brighten and saturate the points a bit so they pop
+				FLinearColor HSV = PointTint.LinearRGBToHSV();
+				HSV.G = FMath::Clamp(HSV.G * 1.1f, 0.f, 255.f);
+				HSV.B = FMath::Clamp(HSV.B * 2.f, 0.f, 255.f);
+				PointTint = HSV.HSVToLinearRGB();
 
 				FPaintGeometry PointGeometry = AllottedGeometry.ToPaintGeometry(
 					Point.ScreenPosition - (PointDrawInfo.ScreenSize * 0.5f),
 					PointDrawInfo.ScreenSize
 				);
 
-				FSlateDrawElement::MakeBox(OutDrawElements, KeyLayerId, PointGeometry, PointDrawInfo.Brush, DrawEffects, PointDrawInfo.Tint);
+				FSlateDrawElement::MakeBox(OutDrawElements, KeyLayerId, PointGeometry, PointDrawInfo.Brush, DrawEffects, PointTint );
 			}
 		}
 	}
diff --git a/Engine/Source/Editor/CurveEditor/Private/Views/SInteractiveCurveEditorView.cpp b/Engine/Source/Editor/CurveEditor/Private/Views/SInteractiveCurveEditorView.cpp
--- a/Engine/Source/Editor/CurveEditor/Private/Views/SInteractiveCurveEditorView.cpp
+++ b/Engine/Source/Editor/CurveEditor/Private/Views/SInteractiveCurveEditorView.cpp
@@ -384,8 +384,7 @@ void SInteractiveCurveEditorView::DrawGridLines(TSharedRef<FCurveEditor> CurveEd
 
 void SInteractiveCurveEditorView::DrawCurves(TSharedRef<FCurveEditor> CurveEditor, const FGeometry& AllottedGeometry, const FSlateRect& MyCullingRect, FSlateWindowElementList& OutDrawElements, int32 BaseLayerId, const FWidgetStyle& InWidgetStyle, ESlateDrawEffect DrawEffects) const
 {
-	static const FName SelectionColorName("SelectionColor");
-	FLinearColor SelectionColor = FAppStyle::GetSlateColor(SelectionColorName).GetColor(InWidgetStyle);
+	FLinearColor SelectionColor = CurveEditor->GetSettings()->GetSelectionColor();
 
 	const FVector2D      VisibleSize = AllottedGeometry.GetLocalSize();
 	const FPaintGeometry PaintGeometry = AllottedGeometry.ToPaintGeometry();
@@ -423,7 +422,20 @@ void SInteractiveCurveEditorView::DrawCurves(TSharedRef<FCurveEditor> CurveEdito
 				const FCurvePointInfo& Point = Params.Points[PointIndex];
 				const FKeyDrawInfo& PointDrawInfo = Params.GetKeyDrawInfo(Point.Type, PointIndex);
 				const bool          bSelected = CurveEditor->GetSelection().IsSelected(FCurvePointHandle(Params.GetID(), Point.Type, Point.KeyHandle));
-				const FLinearColor  PointTint = bSelected ? SelectionColor : PointDrawInfo.Tint;
+				FLinearColor  PointTint = PointDrawInfo.Tint.IsSet() ? PointDrawInfo.Tint.GetValue() : Params.Color;
+
+				if (bSelected)
+				{
+					PointTint = SelectionColor;
+				}
+				else
+				{
+					// Brighten and saturate the points a bit so they pop
+					FLinearColor HSV = PointTint.LinearRGBToHSV();
+					HSV.G = FMath::Clamp(HSV.G * 1.1f, 0.f, 255.f);
+					HSV.B = FMath::Clamp(HSV.B * 2.f, 0.f, 255.f);
+					PointTint = HSV.HSVToLinearRGB();
+				}
 
 				const int32 KeyLayerId = BaseLayerId + Point.LayerBias + (bSelected ? CurveViewConstants::ELayerOffset::SelectedKeys : CurveViewConstants::ELayerOffset::Keys);
 
diff --git a/Engine/Source/Editor/CurveEditor/Public/CurveDrawInfo.h b/Engine/Source/Editor/CurveEditor/Public/CurveDrawInfo.h
--- a/Engine/Source/Editor/CurveEditor/Public/CurveDrawInfo.h
+++ b/Engine/Source/Editor/CurveEditor/Public/CurveDrawInfo.h
@@ -19,7 +19,6 @@ struct FKeyDrawInfo
 	FKeyDrawInfo()
 		: ScreenSize(0.f, 0.f)
 		, Brush(nullptr)
-		, Tint(1.f, 1.f, 1.f, 1.f)
 	{}
 
 	/** The size of the key on screen in slate units */
@@ -29,7 +28,7 @@ struct FKeyDrawInfo
 	const FSlateBrush* Brush;
 
 	/** A tint to apply to the brush */
-	FLinearColor Tint;
+	TOptional<FLinearColor> Tint;
 };
 
 
@@ -88,7 +87,7 @@ struct FCurveDrawParams
 	 * Construct new draw parameters for the specified curve ID
 	 */
 	FCurveDrawParams(FCurveModelID InID)
-		: Color(FLinearColor::White)
+		: Color(0.2f, 0.2f, 0.2f)
 		, bKeyDrawEnabled(1)
 		, ID(InID)
 	{}
diff --git a/Engine/Source/Editor/CurveEditor/Public/CurveEditorSettings.h b/Engine/Source/Editor/CurveEditor/Public/CurveEditorSettings.h
--- a/Engine/Source/Editor/CurveEditor/Public/CurveEditorSettings.h
+++ b/Engine/Source/Editor/CurveEditor/Public/CurveEditorSettings.h
@@ -112,6 +112,11 @@ public:
 	/** Set whether to snap the time to the currently selected key. */
 	void SetSnapTimeToSelection(bool bInSnapTimeToSelection);
 
+	/** Set the selection color. */
+	void SetSelectionColor(const FLinearColor& InColor);
+	/** Get the selection color. */
+	FLinearColor GetSelectionColor() const;
+
 	/** Get custom color for object and property if it exists, if it doesn't the optional won't be set */
 	TOptional<FLinearColor> GetCustomColor(UClass* InClass, const FString& InPropertyName) const;
 	/** Set Custom Color for the specified parameters. */
@@ -156,6 +161,9 @@ protected:
 	UPROPERTY( config, EditAnywhere, Category="Curve Editor")
 	bool bSnapTimeToSelection;
 
+	UPROPERTY(config, EditAnywhere, Category = "Curve Editor")
+	FLinearColor SelectionColor;
+
 	UPROPERTY(config, EditAnywhere, Category="Curve Editor")
 	TArray<FCustomColorForChannel> CustomColors;
 
diff --git a/Engine/Source/Editor/CurveEditor/Public/CurveModel.h b/Engine/Source/Editor/CurveEditor/Public/CurveModel.h
--- a/Engine/Source/Editor/CurveEditor/Public/CurveModel.h
+++ b/Engine/Source/Editor/CurveEditor/Public/CurveModel.h
@@ -38,7 +38,7 @@ class CURVEEDITOR_API FCurveModel
 public:
 
 	FCurveModel()
-		: Color(FLinearColor::White)
+		: Color(0.2f,0.2f,0.2f)
 		, bKeyDrawEnabled(true)
 		, SupportedViews(ECurveEditorViewID::ANY_BUILT_IN)
 	{}
diff --git a/Engine/Source/Editor/CurveTableEditor/Private/RealCurveModel.cpp b/Engine/Source/Editor/CurveTableEditor/Private/RealCurveModel.cpp
--- a/Engine/Source/Editor/CurveTableEditor/Private/RealCurveModel.cpp
+++ b/Engine/Source/Editor/CurveTableEditor/Private/RealCurveModel.cpp
@@ -153,19 +153,15 @@ void FRealCurveModel::GetKeyDrawInfo(ECurvePointType PointType, const FKeyHandle
 		{
 		case ERichCurveInterpMode::RCIM_Constant:
 			OutDrawInfo.Brush = FAppStyle::Get().GetBrush("GenericCurveEditor.ConstantKey");
-			OutDrawInfo.Tint = FLinearColor(0, 0.45f, 0.70f);
 			break;
 		case ERichCurveInterpMode::RCIM_Linear:
 			OutDrawInfo.Brush = FAppStyle::Get().GetBrush("GenericCurveEditor.LinearKey");
-			OutDrawInfo.Tint = FLinearColor(0, 0.62f, 0.46f);
 			break;
 		case ERichCurveInterpMode::RCIM_Cubic:
 			OutDrawInfo.Brush = FAppStyle::Get().GetBrush("GenericCurveEditor.CubicKey");
-			OutDrawInfo.Tint = FLinearColor::White;
 			break;
 		default:
 			OutDrawInfo.Brush = FAppStyle::Get().GetBrush("GenericCurveEditor.Key");
-			OutDrawInfo.Tint = FLinearColor::White;
 			break;
 		}
 	}
diff --git a/Engine/Source/Editor/MovieSceneTools/Private/Channels/BezierChannelCurveModel.cpp b/Engine/Source/Editor/MovieSceneTools/Private/Channels/BezierChannelCurveModel.cpp
--- a/Engine/Source/Editor/MovieSceneTools/Private/Channels/BezierChannelCurveModel.cpp
+++ b/Engine/Source/Editor/MovieSceneTools/Private/Channels/BezierChannelCurveModel.cpp
@@ -114,11 +114,9 @@ void FBezierChannelCurveModel<ChannelType, ChannelValue, KeyType>::GetKeyDrawInf
 		{
 		case ERichCurveInterpMode::RCIM_Constant:
 			OutDrawInfo.Brush = FAppStyle::GetBrush("GenericCurveEditor.ConstantKey");
-			OutDrawInfo.Tint = FLinearColor(0, 0.45f, 0.70f);
 			break;
 		case ERichCurveInterpMode::RCIM_Linear:
 			OutDrawInfo.Brush = FAppStyle::GetBrush("GenericCurveEditor.LinearKey");
-			OutDrawInfo.Tint = FLinearColor(0, 0.62f, 0.46f);
 			break;
 		case ERichCurveInterpMode::RCIM_Cubic:
 			if (KeyTWType == ERichCurveTangentWeightMode::RCTWM_WeightedBoth)
@@ -130,17 +128,15 @@ void FBezierChannelCurveModel<ChannelType, ChannelValue, KeyType>::GetKeyDrawInf
 				OutDrawInfo.Brush = FAppStyle::GetBrush("GenericCurveEditor.CubicKey");
 			}
 
-			OutDrawInfo.Tint = FLinearColor::White;
 			break;
 		default:
 			OutDrawInfo.Brush = FAppStyle::GetBrush("GenericCurveEditor.Key");
-			OutDrawInfo.Tint = FLinearColor::White;
 			break;
 		}
 
 		if (this->IsReadOnly())
 		{
-			OutDrawInfo.Tint = OutDrawInfo.Tint * 0.5f;
+			OutDrawInfo.Tint = OutDrawInfo.Tint.IsSet() ? OutDrawInfo.Tint.GetValue() * 0.5f : FLinearColor(0.5f, 0.5f, 0.5f);
 		}
 	}
 }
diff --git a/Engine/Source/Runtime/Engine/Private/Curves/CurveOwnerInterface.cpp b/Engine/Source/Runtime/Engine/Private/Curves/CurveOwnerInterface.cpp
--- a/Engine/Source/Runtime/Engine/Private/Curves/CurveOwnerInterface.cpp
+++ b/Engine/Source/Runtime/Engine/Private/Curves/CurveOwnerInterface.cpp
@@ -12,18 +12,18 @@ FLinearColor FCurveOwnerInterface::GetCurveColor(FRichCurveEditInfo CurveInfo) c
 
 	if (CurveName == TEXT("X") || CurveName == TEXT("R"))
 	{
-		return FLinearColor(1.0f, 0.0f, 0.0f);
+		return FLinearColor(1.0f, 0.05f, 0.05f);
 	}
 	
 	if (CurveName == TEXT("Y") || CurveName == TEXT("G"))
 	{
-		return FLinearColor(0.0f, 1.0f, 0.0f);
+		return FLinearColor(0.05f, 1.0f, 0.05f);
 	}
 	
 	if (CurveName == TEXT("Z") || CurveName == TEXT("B"))
 	{
-		return FLinearColor(0.05f, 0.05f, 1.0f);
+		return FLinearColor(0.1f, 0.2f, 1.0f);
 	}
 
-	return FLinearColor::Gray;
+	return FLinearColor(0.2f, 0.2f, 0.2f);
 }
diff --git a/Engine/Source/Runtime/MovieScene/Private/Channels/MovieSceneChannelEditorData.cpp b/Engine/Source/Runtime/MovieScene/Private/Channels/MovieSceneChannelEditorData.cpp
--- a/Engine/Source/Runtime/MovieScene/Private/Channels/MovieSceneChannelEditorData.cpp
+++ b/Engine/Source/Runtime/MovieScene/Private/Channels/MovieSceneChannelEditorData.cpp
@@ -16,7 +16,7 @@ const FText FCommonChannelData::ChannelA = NSLOCTEXT("MovieSceneChannels", "Chan
 
 const FLinearColor FCommonChannelData::RedChannelColor(1.0f, 0.05f, 0.05f, 0.9f);
 const FLinearColor FCommonChannelData::GreenChannelColor(0.05f, 1.0f, 0.05f, 0.9f);
-const FLinearColor FCommonChannelData::BlueChannelColor(0.05f, 0.1f, 1.0f, 0.9f);
+const FLinearColor FCommonChannelData::BlueChannelColor(0.1f, 0.2f, 1.0f, 0.9f);
 
 
 FMovieSceneChannelMetaData::FMovieSceneChannelMetaData()
