Curve Editor: Remove any overlapping keys on mouse up (value set)

diff --git a/Engine/Source/Editor/CurveEditor/Private/DragOperations/CurveEditorDragOperation_MoveKeys.cpp b/Engine/Source/Editor/CurveEditor/Private/DragOperations/CurveEditorDragOperation_MoveKeys.cpp
--- a/Engine/Source/Editor/CurveEditor/Private/DragOperations/CurveEditorDragOperation_MoveKeys.cpp
+++ b/Engine/Source/Editor/CurveEditor/Private/DragOperations/CurveEditorDragOperation_MoveKeys.cpp
@@ -136,6 +136,27 @@ void FCurveEditorDragOperation_MoveKeys::OnEndDrag(FVector2D InitialPosition, FV
 	{
 		if (FCurveModel* Curve = CurveEditor->FindCurve(KeyData.CurveID))
 		{
+			// First, remove any keys that overlap before moving the new key on top
+			TArray<FKeyHandle> KeysToRemove;
+			for (int32 KeyIndex = 0; KeyIndex < KeyData.Handles.Num(); ++KeyIndex)
+			{
+				const FKeyHandle& KeyHandle = KeyData.Handles[KeyIndex];
+				FKeyPosition KeyTime = KeyData.LastDraggedKeyPositions[KeyIndex];
+
+				TArray<FKeyHandle> KeysInRange;
+				Curve->GetKeys(*CurveEditor, KeyTime.InputValue, KeyTime.InputValue, TNumericLimits<double>::Lowest(), TNumericLimits<double>::Max(), KeysInRange);
+
+				for (const FKeyHandle& KeyInRangeHandle : KeysInRange)
+				{
+					if (KeyHandle != KeyInRangeHandle)
+					{
+						KeysToRemove.Add(KeyInRangeHandle);
+					}
+				}
+			}
+			Curve->RemoveKeys(KeysToRemove);
+
+			// Then, move the keys
 			Curve->SetKeyPositions(KeyData.Handles, KeyData.LastDraggedKeyPositions, EPropertyChangeType::ValueSet);
 		}
 	}
