Curve Editor: On PostUndo(), clear the selection serial number so that the edit box can refresh and then reset the selection serial number so that time doesn't sync to selection since selection didn't really change

diff --git a/Engine/Source/Editor/CurveEditor/Private/SCurveEditorPanel.cpp b/Engine/Source/Editor/CurveEditor/Private/SCurveEditorPanel.cpp
--- a/Engine/Source/Editor/CurveEditor/Private/SCurveEditorPanel.cpp
+++ b/Engine/Source/Editor/CurveEditor/Private/SCurveEditorPanel.cpp
@@ -451,7 +451,13 @@ void SCurveEditorPanel::RemoveCurveFromViews(FCurveModelID InCurveID)
 void SCurveEditorPanel::PostUndo()
 {
 	EditObjects->CurveIDToKeyProxies.Empty();
+
+	// Force the edit box to update (ie. the value of the keys might have changed)
 	CachedSelectionSerialNumber = 0;
+	UpdateEditBox();
+
+	// Reset the selection serial number so that time doesn't change since selection didn't really change on undo
+	CachedSelectionSerialNumber = CurveEditor->Selection.GetSerialNumber();
 }
 
 void SCurveEditorPanel::AddView(TSharedRef<SCurveEditorView> ViewToAdd)
