Curve Editor: Hook up bindings for ZoomToFitAll, ZoomToFitHorizontal, ZoomToFitVertical

diff --git a/Engine/Source/Editor/CurveEditor/Private/CurveEditor.cpp b/Engine/Source/Editor/CurveEditor/Private/CurveEditor.cpp
--- a/Engine/Source/Editor/CurveEditor/Private/CurveEditor.cpp
+++ b/Engine/Source/Editor/CurveEditor/Private/CurveEditor.cpp
@@ -350,6 +350,9 @@ void FCurveEditor::BindCommands()
 	CommandList->MapAction(FGenericCommands::Get().Paste, FExecuteAction::CreateSP(this, &FCurveEditor::PasteKeys, TSet<FCurveModelID>()));
 
 	CommandList->MapAction(FCurveEditorCommands::Get().ZoomToFit, FExecuteAction::CreateSP(this, &FCurveEditor::ZoomToFit, EAxisList::All));
+	CommandList->MapAction(FCurveEditorCommands::Get().ZoomToFitHorizontal, FExecuteAction::CreateSP(this, &FCurveEditor::ZoomToFit, EAxisList::X));
+	CommandList->MapAction(FCurveEditorCommands::Get().ZoomToFitVertical, FExecuteAction::CreateSP(this, &FCurveEditor::ZoomToFit, EAxisList::Y));
+	CommandList->MapAction(FCurveEditorCommands::Get().ZoomToFitAll, FExecuteAction::CreateSP(this, &FCurveEditor::ZoomToFitAll, EAxisList::All));
 
 	CommandList->MapAction(FCurveEditorCommands::Get().ToggleExpandCollapseNodes, FExecuteAction::CreateSP(this, &FCurveEditor::ToggleExpandCollapseNodes, false));
 	CommandList->MapAction(FCurveEditorCommands::Get().ToggleExpandCollapseNodesAndDescendants, FExecuteAction::CreateSP(this, &FCurveEditor::ToggleExpandCollapseNodes, true));
@@ -488,15 +491,20 @@ void FCurveEditor::ZoomToFit(EAxisList::Type Axes)
 	}
 	else
 	{
-		TMap<FCurveModelID, FKeyHandleSet> AllCurves;
-		for (FCurveModelID ID : GetEditedCurves())
-		{
-			AllCurves.Add(ID);
-		}
-		ZoomToFitInternal(Axes, AllCurves);
+		ZoomToFitAll(Axes);
 	}
 }
 
+void FCurveEditor::ZoomToFitAll(EAxisList::Type Axes)
+{
+	TMap<FCurveModelID, FKeyHandleSet> AllCurves;
+	for (FCurveModelID ID : GetEditedCurves())
+	{
+		AllCurves.Add(ID);
+	}
+	ZoomToFitInternal(Axes, AllCurves);
+}
+
 void FCurveEditor::ZoomToFitCurves(TArrayView<const FCurveModelID> CurveModelIDs, EAxisList::Type Axes)
 {
 	TMap<FCurveModelID, FKeyHandleSet> AllCurves;
diff --git a/Engine/Source/Editor/CurveEditor/Public/CurveEditor.h b/Engine/Source/Editor/CurveEditor/Public/CurveEditor.h
--- a/Engine/Source/Editor/CurveEditor/Public/CurveEditor.h
+++ b/Engine/Source/Editor/CurveEditor/Public/CurveEditor.h
@@ -118,11 +118,17 @@ public:
 
 public:
 	/**
-	 * Zoom the curve editor to fit all the currently visible curves
+	 * Zoom the curve editor to fit all the selected curves (or all curves if none selected)
 	 *
 	 * @param Axes         (Optional) Axes to lock the zoom to
 	 */
 	void ZoomToFit(EAxisList::Type Axes = EAxisList::All);
+	/**
+	 * Zoom the curve editor to fit all the currently visible curves
+	 *
+	 * @param Axes         (Optional) Axes to lock the zoom to
+	 */
+	void ZoomToFitAll(EAxisList::Type Axes = EAxisList::All);
 	/**
 	 * Zoom the curve editor to fit the requested curves.
 	 *
