Curve Editor: Fix crash when manipulating tangents

LastDraggedAttributes might not be set if nothing has been dragged.

Seems to be fallout from UE-150118 CL#19971467 but I haven't been able to reproduce it.

diff --git a/Engine/Source/Editor/CurveEditor/Private/DragOperations/CurveEditorDragOperation_Tangent.cpp b/Engine/Source/Editor/CurveEditor/Private/DragOperations/CurveEditorDragOperation_Tangent.cpp
--- a/Engine/Source/Editor/CurveEditor/Private/DragOperations/CurveEditorDragOperation_Tangent.cpp
+++ b/Engine/Source/Editor/CurveEditor/Private/DragOperations/CurveEditorDragOperation_Tangent.cpp
@@ -203,7 +203,10 @@ void FCurveEditorDragOperation_Tangent::OnEndDrag(FVector2D InitialPosition, FVe
 	{
 		if (FCurveModel* Curve = CurveEditor->FindCurve(KeyData.CurveID))
 		{
-			Curve->SetKeyAttributes(KeyData.Handles, KeyData.LastDraggedAttributes, EPropertyChangeType::ValueSet);
+			if (KeyData.LastDraggedAttributes.IsSet())
+			{
+				Curve->SetKeyAttributes(KeyData.Handles, KeyData.LastDraggedAttributes.GetValue(), EPropertyChangeType::ValueSet);
+			}
 		}
 	}
 	
diff --git a/Engine/Source/Editor/CurveEditor/Private/DragOperations/CurveEditorDragOperation_Tangent.h b/Engine/Source/Editor/CurveEditor/Private/DragOperations/CurveEditorDragOperation_Tangent.h
--- a/Engine/Source/Editor/CurveEditor/Private/DragOperations/CurveEditorDragOperation_Tangent.h
+++ b/Engine/Source/Editor/CurveEditor/Private/DragOperations/CurveEditorDragOperation_Tangent.h
@@ -45,7 +45,7 @@ private:
 		/** The key attributes for each of the above handles */
 		TArray<FKeyAttributes> Attributes;
 		/** Used in OnEndDrag to send final key updates */
-		TArray<FKeyAttributes> LastDraggedAttributes;
+		TOptional<TArray<FKeyAttributes> > LastDraggedAttributes;
 	};
 
 	/** Key dragging data stored per-curve */
