Curve Editor: When shift and marquee select, make sure only points are selected if points are selected. Or, make sure tangents are selected if only tangents are selected.

diff --git a/Engine/Source/Editor/CurveEditor/Private/DragOperations/CurveEditorDragOperation_Marquee.cpp b/Engine/Source/Editor/CurveEditor/Private/DragOperations/CurveEditorDragOperation_Marquee.cpp
--- a/Engine/Source/Editor/CurveEditor/Private/DragOperations/CurveEditorDragOperation_Marquee.cpp
+++ b/Engine/Source/Editor/CurveEditor/Private/DragOperations/CurveEditorDragOperation_Marquee.cpp
@@ -101,6 +101,41 @@ void FCurveEditorDragOperation_Marquee::OnEndDrag(FVector2D InitialPosition, FVe
 		}
 	}
 
+	// When adding to the existing selection, ensure only either points or tangents are selected
+	TArray<FCurvePointHandle> CurvePointsToRemove;
+	if (bIsShiftDown)
+	{
+		for (const TTuple<FCurveModelID, FKeyHandleSet>& Pair : CurveEditor->GetSelection().GetAll())
+		{
+			for (FKeyHandle Handle : Pair.Value.AsArray())
+			{
+				ECurvePointType PointType = Pair.Value.PointType(Handle);
+
+				if (bPreferPointSelection)
+				{
+					// If selecting points, deselect tangen handles (ie. anything that's not a point/key)
+					if (PointType != ECurvePointType::Key)
+					{
+						CurvePointsToRemove.Add(FCurvePointHandle(Pair.Key, PointType, Handle));
+					}
+				}
+				else
+				{
+					// Otherwise when selecting tangent handles, deselect anything that's a key
+					if (PointType == ECurvePointType::Key)
+					{
+						CurvePointsToRemove.Add(FCurvePointHandle(Pair.Key, PointType, Handle));
+					}
+				}
+			}
+		}
+	}
+
+	for (const FCurvePointHandle& Point : CurvePointsToRemove)
+	{
+		CurveEditor->Selection.Remove(Point);
+	}
+
 	// Now that we've gathered the overlapping points, perform the relevant selection
 	for (const FCurvePointHandle& Point : AllPoints)
 	{
