Curve Editor: Snap time regardless of view mode.

Snap time and value was only uccuring for Absolute view mode, but time should always be snapped regardless.

diff --git a/Engine/Source/Editor/CurveEditor/Private/DragOperations/CurveEditorDragOperation_MoveKeys.cpp b/Engine/Source/Editor/CurveEditor/Private/DragOperations/CurveEditorDragOperation_MoveKeys.cpp
--- a/Engine/Source/Editor/CurveEditor/Private/DragOperations/CurveEditorDragOperation_MoveKeys.cpp
+++ b/Engine/Source/Editor/CurveEditor/Private/DragOperations/CurveEditorDragOperation_MoveKeys.cpp
@@ -71,8 +71,7 @@ void FCurveEditorDragOperation_MoveKeys::OnDrag(FVector2D InitialPosition, FVect
 
 		FCurveSnapMetrics SnapMetrics = CurveEditor->GetCurveSnapMetrics(KeyData.CurveID);
 
-		// If view is not absolute, snap based on the key that was grabbed, not all keys individually.
-		if (CardinalPoint.IsSet() && (View->IsTimeSnapEnabled() || View->IsValueSnapEnabled()) && View->ViewTypeID != ECurveEditorViewID::Absolute)
+		if (CardinalPoint.IsSet())
 		{
 			for (int KeyIndex = 0; KeyIndex < KeyData.StartKeyPositions.Num(); ++KeyIndex)
 			{
@@ -85,7 +84,9 @@ void FCurveEditorDragOperation_MoveKeys::OnDrag(FVector2D InitialPosition, FVect
 					{
 						DeltaInput = SnapMetrics.SnapInputSeconds(StartPosition.InputValue + DeltaInput) - StartPosition.InputValue;
 					}
-					if (View->IsValueSnapEnabled())
+
+					// If view is not absolute, snap based on the key that was grabbed, not all keys individually.
+					if (View->IsValueSnapEnabled() && View->ViewTypeID != ECurveEditorViewID::Absolute)
 					{
 						DeltaOutput = SnapMetrics.SnapOutput(StartPosition.OutputValue + DeltaOutput) - StartPosition.OutputValue;
 					}
@@ -98,10 +99,11 @@ void FCurveEditorDragOperation_MoveKeys::OnDrag(FVector2D InitialPosition, FVect
 			StartPosition.InputValue  += DeltaInput;
 			StartPosition.OutputValue += DeltaOutput;
 
-			// Snap keys individually if view mode is absolute.
+			StartPosition.InputValue = View->IsTimeSnapEnabled() ? SnapMetrics.SnapInputSeconds(StartPosition.InputValue) : StartPosition.InputValue;
+
+			// Snap value keys individually if view mode is absolute.
 			if (View->ViewTypeID == ECurveEditorViewID::Absolute)
 			{
-				StartPosition.InputValue  = View->IsTimeSnapEnabled() ? SnapMetrics.SnapInputSeconds(StartPosition.InputValue) : StartPosition.InputValue;
 				StartPosition.OutputValue = View->IsValueSnapEnabled() ? SnapMetrics.SnapOutput(StartPosition.OutputValue) : StartPosition.OutputValue;
 			}
 			
