Curve Editor: Remove duplicate bound shortcut

diff --git a/Engine/Source/Editor/CurveEditor/Private/CurveEditorCommands.cpp b/Engine/Source/Editor/CurveEditor/Private/CurveEditorCommands.cpp
--- a/Engine/Source/Editor/CurveEditor/Private/CurveEditorCommands.cpp
+++ b/Engine/Source/Editor/CurveEditor/Private/CurveEditorCommands.cpp
@@ -14,7 +14,7 @@ void FCurveEditorCommands::RegisterCommands()
 	UI_COMMAND(ZoomToFitHorizontal, "Fit Horizontal", "Zoom to Fit - Horizontal", EUserInterfaceActionType::Button, FInputChord());
 	UI_COMMAND(ZoomToFitVertical, "Fit Vertical", "Zoom to Fit - Vertical", EUserInterfaceActionType::Button, FInputChord());
 	UI_COMMAND(ZoomToFit, "Fit", "Zoom to Fit", EUserInterfaceActionType::Button, FInputChord(EKeys::F));
-	UI_COMMAND(ZoomToFitAll, "FitAll", "Zoom to Fit All", EUserInterfaceActionType::Button, FInputChord(EKeys::A));
+	UI_COMMAND(ZoomToFitAll, "FitAll", "Zoom to Fit All", EUserInterfaceActionType::Button, FInputChord());
 
 	UI_COMMAND(ToggleInputSnapping, "Input Snapping", "Toggle Time Snapping", EUserInterfaceActionType::ToggleButton, FInputChord());
 	UI_COMMAND(ToggleOutputSnapping, "Output Snapping", "Toggle Value Snapping", EUserInterfaceActionType::ToggleButton, FInputChord());
