Curve Editor: Auto frame after paste operation

diff --git a/Engine/Source/Editor/CurveEditor/Private/CurveEditor.cpp b/Engine/Source/Editor/CurveEditor/Private/CurveEditor.cpp
--- a/Engine/Source/Editor/CurveEditor/Private/CurveEditor.cpp
+++ b/Engine/Source/Editor/CurveEditor/Private/CurveEditor.cpp
@@ -1582,6 +1582,8 @@ void FCurveEditor::PasteKeys(TSet<FCurveModelID> CurveModelIDs)
 			}
 		}
 	}
+
+	ZoomToFitSelection();
 }
 
 void FCurveEditor::DeleteSelection()
