Curve Editor: Autoframe on paste only when setting is enabled

diff --git a/Engine/Source/Editor/CurveEditor/Private/CurveEditor.cpp b/Engine/Source/Editor/CurveEditor/Private/CurveEditor.cpp
--- a/Engine/Source/Editor/CurveEditor/Private/CurveEditor.cpp
+++ b/Engine/Source/Editor/CurveEditor/Private/CurveEditor.cpp
@@ -1583,7 +1583,10 @@ void FCurveEditor::PasteKeys(TSet<FCurveModelID> CurveModelIDs)
 		}
 	}
 
-	ZoomToFitSelection();
+	if (ShouldAutoFrame())
+	{
+		ZoomToFitSelection();
+	}
 }
 
 void FCurveEditor::DeleteSelection()
