Curve Editor: Curve color tweaks

- Set color for all selected curves in the tree and/or curves that have selected keyframes (ie. curve selection)
- Fix next random color - pastel green
- Broadcast when custom colors changed so that the curve models can update their colors

Note - if you delete a custom color from the preferences, it doesn't refresh the curve editor as you might expect because the color is no longer defined in the key area meta data. The only way to fix this is to recreate the curve models, but that seems like overkill for this (selection would change, framing would change, etc)

diff --git a/Engine/Source/Editor/CurveEditor/Private/CurveEditor.cpp b/Engine/Source/Editor/CurveEditor/Private/CurveEditor.cpp
--- a/Engine/Source/Editor/CurveEditor/Private/CurveEditor.cpp
+++ b/Engine/Source/Editor/CurveEditor/Private/CurveEditor.cpp
@@ -84,6 +84,16 @@ FCurveEditor::FCurveEditor()
 
 	GridLineLabelFormatXAttribute = LOCTEXT("GridXLabelFormat", "{0}s");
 	GridLineLabelFormatYAttribute = LOCTEXT("GridYLabelFormat", "{0}");
+	
+	Settings->GetOnCustomColorsChanged().AddRaw(this, &FCurveEditor::OnCustomColorsChanged);
+}
+
+FCurveEditor::~FCurveEditor()
+{
+	if (Settings)
+	{
+		Settings->GetOnCustomColorsChanged().RemoveAll(this);
+	}
 }
 
 void FCurveEditor::InitCurveEditor(const FCurveEditorInitParams& InInitParams)
@@ -1751,9 +1761,15 @@ bool FCurveEditor::CanFlattenOrStraightenSelection() const
 
 void FCurveEditor::SetRandomCurveColorsForSelected()
 {
-	for (TPair<FCurveModelID, TUniquePtr<FCurveModel>>& CurvePair : CurveData)
+	TSet<FCurveModelID> CurveModelIDs = GetSelectionFromTreeAndKeys();
+	if (CurveModelIDs.Num() == 0)
 	{
-		if (FCurveModel* Curve = CurvePair.Value.Get())
+		return;
+	}
+
+	for (const FCurveModelID& CurveModelID : CurveModelIDs)
+	{
+		if (FCurveModel* Curve = FindCurve(CurveModelID))
 		{
 			UObject* Object = nullptr;
 			FString Name;
@@ -1770,31 +1786,33 @@ void FCurveEditor::SetRandomCurveColorsForSelected()
 
 void FCurveEditor::SetCurveColorsForSelected()
 {
-	if (CurveData.Num() > 0)
+	TSet<FCurveModelID> CurveModelIDs = GetSelectionFromTreeAndKeys();
+	if (CurveModelIDs.Num() == 0)
 	{
-		TMap<FCurveModelID, TUniquePtr<FCurveModel>>::TIterator It = CurveData.CreateIterator();
-		FColorPickerArgs PickerArgs;
-		PickerArgs.bUseAlpha = false;
-		PickerArgs.InitialColorOverride = It->Value->GetColor();
-		PickerArgs.OnColorCommitted.BindLambda([this](FLinearColor NewColor) {
-			for (TPair<FCurveModelID, TUniquePtr<FCurveModel>>& CurvePair : CurveData)
+		return;
+	}
+
+	FColorPickerArgs PickerArgs;
+	PickerArgs.bUseAlpha = false;
+	PickerArgs.InitialColorOverride = FindCurve(*CurveModelIDs.CreateIterator())->GetColor();
+	PickerArgs.OnColorCommitted.BindLambda([this, CurveModelIDs](FLinearColor NewColor) {
+		for (const FCurveModelID& CurveModelID : CurveModelIDs)
+		{
+			if (FCurveModel* Curve = FindCurve(CurveModelID))
 			{
-				if (FCurveModel* Curve = CurvePair.Value.Get())
+				UObject* Object = nullptr;
+				FString Name;
+				Curve->GetCurveColorObjectAndName(&Object, Name);
+				if (Object)
 				{
-					UObject* Object = nullptr;
-					FString Name;
-					Curve->GetCurveColorObjectAndName(&Object, Name);
-					if (Object)
-					{
-						Settings->SetCustomColor(Object->GetClass(), Name, NewColor);
-						Curve->SetColor(NewColor);
-					}
+					Settings->SetCustomColor(Object->GetClass(), Name, NewColor);
+					Curve->SetColor(NewColor);
 				}
 			}
-		});
+		}
+	});
 		
-		OpenColorPicker(PickerArgs);
-	}
+	OpenColorPicker(PickerArgs);
 }
 
 bool FCurveEditor::IsToolActive(const FCurveEditorToolID InToolID) const
@@ -2099,7 +2117,7 @@ bool FCurveEditor::ApplyBufferedCurves(const TSet<FCurveModelID>& InCurvesToAppl
 	return false;
 }
 
-TSet<FCurveModelID> FCurveEditor::GetCurvesForBufferedCurves() const
+TSet<FCurveModelID> FCurveEditor::GetSelectionFromTreeAndKeys() const
 {
 	TSet<FCurveModelID> CurveModelIDs;
 
@@ -2126,7 +2144,7 @@ TSet<FCurveModelID> FCurveEditor::GetCurvesForBufferedCurves() const
 
 bool FCurveEditor::IsActiveBufferedCurve(const TUniquePtr<IBufferedCurveModel>& BufferedCurve) const
 {
-	TSet<FCurveModelID> CurveModelIDs = GetCurvesForBufferedCurves();
+	TSet<FCurveModelID> CurveModelIDs = GetSelectionFromTreeAndKeys();
 	for (const FCurveModelID& CurveModelID : CurveModelIDs)
 	{
 		if (FCurveModel* Curve = FindCurve(CurveModelID))
@@ -2182,4 +2200,29 @@ void FCurveEditor::PostUndo(bool bSuccess)
 	}
 }
 
+void FCurveEditor::OnCustomColorsChanged()
+{
+	for (TPair<FCurveModelID, TUniquePtr<FCurveModel>>& CurvePair : CurveData)
+	{
+		if (FCurveModel* Curve = CurvePair.Value.Get())
+		{
+			UObject* Object = nullptr;
+			FString Name;
+			Curve->GetCurveColorObjectAndName(&Object, Name);
+
+			TOptional<FLinearColor> Color = Settings->GetCustomColor(Object->GetClass(), Name);
+			if (Color.IsSet())
+			{
+				Curve->SetColor(Color.GetValue());
+			}
+			else
+			{
+				// Note: If the color is no longer defined, there's no way to update with the previously defined 
+				// default color. The curve models would need to be rebuilt, but would cause selection/framing and 
+				// other things to change. So, this is intentionally not implemented.
+			}
+		}
+	}
+}
+
 #undef LOCTEXT_NAMESPACE
diff --git a/Engine/Source/Editor/CurveEditor/Private/CurveEditorContextMenu.cpp b/Engine/Source/Editor/CurveEditor/Private/CurveEditorContextMenu.cpp
--- a/Engine/Source/Editor/CurveEditor/Private/CurveEditorContextMenu.cpp
+++ b/Engine/Source/Editor/CurveEditor/Private/CurveEditorContextMenu.cpp
@@ -29,7 +29,7 @@ void FCurveEditorContextMenu::BuildMenu(FMenuBuilder& MenuBuilder, TSharedRef<FC
 
 	TSharedPtr<FCurveEditor> LocalCurveEditor = CurveEditor->AsShared();
 
-	int32 NumActiveCurves = LocalCurveEditor->GetCurvesForBufferedCurves().Num();
+	int32 NumActiveCurves = LocalCurveEditor->GetSelectionFromTreeAndKeys().Num();
 
 	int32 NumBufferedCurves = 0;
 	for (const TUniquePtr<IBufferedCurveModel>& BufferedCurve : LocalCurveEditor->GetBufferedCurves())
diff --git a/Engine/Source/Editor/CurveEditor/Private/CurveEditorSettings.cpp b/Engine/Source/Editor/CurveEditor/Private/CurveEditorSettings.cpp
--- a/Engine/Source/Editor/CurveEditor/Private/CurveEditorSettings.cpp
+++ b/Engine/Source/Editor/CurveEditor/Private/CurveEditorSettings.cpp
@@ -277,20 +277,33 @@ FLinearColor UCurveEditorSettings::GetNextRandomColor()
 	if (IndexedColor.Num() == 0)
 	{
 		IndexedColor.Add(FLinearColor(.904, .323, .539)); //pastel red
-		IndexedColor.Add(FLinearColor(552, .737, .328)); //pastel green
+		IndexedColor.Add(FLinearColor(.552, .737, .328)); //pastel green
 		IndexedColor.Add(FLinearColor(.947, .418, .219)); //pastel orange
 		IndexedColor.Add(FLinearColor(.156, .624, .921)); //pastel blue
-		IndexedColor.Add(FLinearColor(.921, .314, .337));//pastel red 2
+		IndexedColor.Add(FLinearColor(.921, .314, .337)); //pastel red 2
 		IndexedColor.Add(FLinearColor(.361, .651, .332)); //pastel green 2
 		IndexedColor.Add(FLinearColor(.982, .565, .254)); //pastel orange 2
 		IndexedColor.Add(FLinearColor(.246, .223, .514)); //pastel purple
 		IndexedColor.Add(FLinearColor(.208, .386, .687)); //pastel blue2
-		IndexedColor.Add(FLinearColor(.223, .59, .337)); //pastel green 3
-		IndexedColor.Add(FLinearColor(.230, .291, .591));  //pastel blue 3
+		IndexedColor.Add(FLinearColor(.223, .590, .337)); //pastel green 3
+		IndexedColor.Add(FLinearColor(.230, .291, .591)); //pastel blue 3
 	}
 	FLinearColor Color = IndexedColor[NextIndex];
 	++NextIndex;
 	int32 NewIndex = (NextIndex % IndexedColor.Num());
 	NextIndex = NewIndex;
 	return Color;
-}
\ No newline at end of file
+}
+
+void UCurveEditorSettings::PostEditChangeProperty(FPropertyChangedEvent& PropertyChangedEvent)
+{
+	FName PropertyName = (PropertyChangedEvent.Property != NULL) ? PropertyChangedEvent.Property->GetFName() : NAME_None;
+	FName MemberPropertyName = (PropertyChangedEvent.MemberProperty != NULL) ? PropertyChangedEvent.MemberProperty->GetFName() : NAME_None;
+	if (PropertyName == GET_MEMBER_NAME_CHECKED(UCurveEditorSettings, CustomColors) ||
+		MemberPropertyName == GET_MEMBER_NAME_CHECKED(UCurveEditorSettings, CustomColors))
+	{
+		OnCustomColorsChangedEvent.Broadcast();
+	}
+	
+	Super::PostEditChangeProperty(PropertyChangedEvent);
+}
diff --git a/Engine/Source/Editor/CurveEditor/Private/Views/SInteractiveCurveEditorView.cpp b/Engine/Source/Editor/CurveEditor/Private/Views/SInteractiveCurveEditorView.cpp
--- a/Engine/Source/Editor/CurveEditor/Private/Views/SInteractiveCurveEditorView.cpp
+++ b/Engine/Source/Editor/CurveEditor/Private/Views/SInteractiveCurveEditorView.cpp
@@ -1399,7 +1399,7 @@ void SInteractiveCurveEditorView::BufferCurves()
 	TSharedPtr<FCurveEditor> CurveEditor = WeakCurveEditor.Pin();
 	if (CurveEditor.IsValid())
 	{
-		CurveEditor->AddBufferedCurves(CurveEditor->GetCurvesForBufferedCurves());
+		CurveEditor->AddBufferedCurves(CurveEditor->GetSelectionFromTreeAndKeys());
 	}
 }
 
@@ -1408,7 +1408,7 @@ void SInteractiveCurveEditorView::ApplyBufferCurves(const bool bSwapBufferCurves
 	TSharedPtr<FCurveEditor> CurveEditor = WeakCurveEditor.Pin();
 	if (CurveEditor.IsValid())
 	{
-		CurveEditor->ApplyBufferedCurves(CurveEditor->GetCurvesForBufferedCurves(), bSwapBufferCurves);
+		CurveEditor->ApplyBufferedCurves(CurveEditor->GetSelectionFromTreeAndKeys(), bSwapBufferCurves);
 	}
 }
 
@@ -1417,7 +1417,7 @@ bool SInteractiveCurveEditorView::CanBufferedCurves() const
 	TSharedPtr<FCurveEditor> CurveEditor = WeakCurveEditor.Pin();
 	if (CurveEditor.IsValid())
 	{
-		return CurveEditor->GetCurvesForBufferedCurves().Num() > 0;
+		return CurveEditor->GetSelectionFromTreeAndKeys().Num() > 0;
 	}
 
 	return false;
@@ -1428,7 +1428,7 @@ bool SInteractiveCurveEditorView::CanApplyBufferedCurves() const
 	TSharedPtr<FCurveEditor> CurveEditor = WeakCurveEditor.Pin();
 	if (CurveEditor.IsValid())
 	{
-		return CurveEditor->GetCurvesForBufferedCurves().Num() > 0 && CurveEditor->GetBufferedCurves().Num() > 0;
+		return CurveEditor->GetSelectionFromTreeAndKeys().Num() > 0 && CurveEditor->GetBufferedCurves().Num() > 0;
 	}
 
 	return false;
diff --git a/Engine/Source/Editor/CurveEditor/Public/CurveEditor.h b/Engine/Source/Editor/CurveEditor/Public/CurveEditor.h
--- a/Engine/Source/Editor/CurveEditor/Public/CurveEditor.h
+++ b/Engine/Source/Editor/CurveEditor/Public/CurveEditor.h
@@ -116,7 +116,7 @@ public:
 	FCurveEditor(const FCurveEditor&) = delete;
 	FCurveEditor& operator=(const FCurveEditor&) = delete;
 
-	virtual ~FCurveEditor() {}
+	virtual ~FCurveEditor();
 
 	void InitCurveEditor(const FCurveEditorInitParams& InInitParams);
 
@@ -294,6 +294,9 @@ public:
 	void SuppressBoundTransformUpdates(bool bSuppress) { bBoundTransformUpdatesSuppressed = bSuppress; }
 	bool AreBoundTransformUpdatesSuppressed() const { return bBoundTransformUpdatesSuppressed; }
 
+	/** Return the curve model IDs that are selected in the tree or have selected keys */
+	TSet<FCurveModelID> GetSelectionFromTreeAndKeys() const;
+
 	TSet<FCurveModelID> GetEditedCurves() const;
 	/** Attribute used for determining default attributes to apply to a newly create key */
 	TAttribute<FKeyAttributes> GetDefaultKeyAttribute() const { return DefaultKeyAttributes; }
@@ -305,8 +308,6 @@ public:
 	int32 GetNumBufferedCurves() const { return BufferedCurves.Num(); }
 	/** Return the array of buffered curves */
 	const TArray<TUniquePtr<IBufferedCurveModel>>& GetBufferedCurves() const { return BufferedCurves; }
-	/** Return the curve model IDs to operate buffer curves on, ie. selected in the tree view or with selected keys */
-	TSet<FCurveModelID> GetCurvesForBufferedCurves() const;
 	/** Returns whether the buffered curve is to be acted on, ie. selected, in the tree view or with selected keys */
 	bool IsActiveBufferedCurve(const TUniquePtr<IBufferedCurveModel>& BufferedCurve) const;
 
@@ -577,6 +578,7 @@ protected:
 	*/
 	void ApplyBufferedCurveToTarget(const IBufferedCurveModel* BufferedCurve, FCurveModel* TargetCurve);
 
+	void OnCustomColorsChanged();
 
 protected:
 
@@ -644,5 +646,4 @@ public:
 	* Delegate that's broadcast when the curve display changes.
 	*/
 	FOnCurveArrayChanged OnCurveArrayChanged;
-
 };
diff --git a/Engine/Source/Editor/CurveEditor/Public/CurveEditorSettings.h b/Engine/Source/Editor/CurveEditor/Public/CurveEditorSettings.h
--- a/Engine/Source/Editor/CurveEditor/Public/CurveEditorSettings.h
+++ b/Engine/Source/Editor/CurveEditor/Public/CurveEditorSettings.h
@@ -80,6 +80,8 @@ class CURVEEDITOR_API UCurveEditorSettings : public UObject
 public:
 	GENERATED_BODY()
 
+	DECLARE_MULTICAST_DELEGATE(FOnCustomColorsChanged);
+
 	UCurveEditorSettings();
 
 	/** Gets whether or not the curve editor auto frames the selected curves. */
@@ -133,6 +135,9 @@ public:
 	void SetCustomColor(UClass* InClass, const FString& InPropertyName, FLinearColor InColor);
 	/** Delete Custom Color for the specified parameters. */
 	void DeleteCustomColor(UClass* InClass, const FString& InPropertyName);
+	
+	/** Gets the multicast delegate which is run whenever custom colors have changed. */
+	FOnCustomColorsChanged& GetOnCustomColorsChanged() { return OnCustomColorsChangedEvent; }
 
 	/** Get custom color for space name. Parent and World are reserved names and will be used instead of the specified control name. */
 	TOptional<FLinearColor> GetSpaceSwitchColor(const FString& InControlName) const;
@@ -145,6 +150,9 @@ public:
 	static FLinearColor GetNextRandomColor();
 
 protected:
+
+	virtual void PostEditChangeProperty(FPropertyChangedEvent& PropertyChangedEvent) override;
+
 	UPROPERTY( config, EditAnywhere, Category="Curve Editor" )
 	bool bAutoFrameCurveEditor;
 
@@ -185,4 +193,7 @@ protected:
 
 	UPROPERTY(config, EditAnywhere, Category = "Curve Editor")
 	TArray<FCustomColorForSpaceSwitch> ControlSpaceCustomColors;
+
+private:
+	FOnCustomColorsChanged OnCustomColorsChangedEvent;
 };
