Curve Editor: Rename a few items from "Store Curves" to "Buffer Curves"

diff --git a/Engine/Source/Editor/CurveEditor/Private/CurveEditorCommands.cpp b/Engine/Source/Editor/CurveEditor/Private/CurveEditorCommands.cpp
--- a/Engine/Source/Editor/CurveEditor/Private/CurveEditorCommands.cpp
+++ b/Engine/Source/Editor/CurveEditor/Private/CurveEditorCommands.cpp
@@ -88,9 +88,9 @@ void FCurveEditorCommands::RegisterCommands()
 
 	// Buffer and Apply Curves. Like copy and paste, but directly onto the curves they were stored from.
 	// These names are overwritten in CurveEditorContextMenu to show the number of stashed curves.
-	UI_COMMAND(BufferVisibleCurves, "Store Curves", "Stores a copy of the selected curves which can be applied back onto themselves.", EUserInterfaceActionType::Button, FInputChord());
-	UI_COMMAND(SwapBufferedCurves, "Swap Stored Curves", "Applies the stored curves to the curves they were stored from and stores the current curves to the buffer.", EUserInterfaceActionType::Button, FInputChord());
-	UI_COMMAND(ApplyBufferedCurves, "Apply Stored Curves", "Applies the stored curves to the curves they were stored from.", EUserInterfaceActionType::Button, FInputChord());
+	UI_COMMAND(BufferVisibleCurves, "Buffer Curves", "Stores a copy of the selected curves which can be applied back onto themselves.", EUserInterfaceActionType::Button, FInputChord());
+	UI_COMMAND(SwapBufferedCurves, "Swap Buffered Curves", "Applies the buffered curves to the curves they were stored from and stores the current curves to the buffer.", EUserInterfaceActionType::Button, FInputChord());
+	UI_COMMAND(ApplyBufferedCurves, "Apply Buffered Curves", "Applies the buffered curves to the curves they were stored from.", EUserInterfaceActionType::Button, FInputChord());
 
 	// Axis Snapping
 	UI_COMMAND(SetAxisSnappingNone, "Both", "Disable axis snapping and allow movement on both the X and Y directions.", EUserInterfaceActionType::Button, FInputChord());
diff --git a/Engine/Source/Editor/CurveEditor/Private/CurveEditorContextMenu.cpp b/Engine/Source/Editor/CurveEditor/Private/CurveEditorContextMenu.cpp
--- a/Engine/Source/Editor/CurveEditor/Private/CurveEditorContextMenu.cpp
+++ b/Engine/Source/Editor/CurveEditor/Private/CurveEditorContextMenu.cpp
@@ -41,9 +41,9 @@ void FCurveEditorContextMenu::BuildMenu(FMenuBuilder& MenuBuilder, TSharedRef<FC
 	}
 
 	// We change the name to reflect the current number of curves selected.
-	TAttribute<FText> BufferedCurvesText = FText::Format(LOCTEXT("StoreCurvesContextMenu", "Store {0} Curves"), NumActiveCurves);
-	TAttribute<FText> SwapBufferedCurvesText = FText::Format(LOCTEXT("SwapStoredCurvesContextMenu", "Swap {0} Stored Curves onto {1} Selected Curves"), NumBufferedCurves, NumActiveCurves);
-	TAttribute<FText> ApplyBufferedCurvesText = FText::Format(LOCTEXT("ApplyStoredCurvesContextMenu", "Apply {0} Stored Curves onto {1} Selected Curves"), NumBufferedCurves, NumActiveCurves);
+	TAttribute<FText> BufferedCurvesText = FText::Format(LOCTEXT("BufferCurvesContextMenu", "Buffer {0} Curves"), NumActiveCurves);
+	TAttribute<FText> SwapBufferedCurvesText = FText::Format(LOCTEXT("SwapBufferedCurvesContextMenu", "Swap {0} Buffered Curves onto {1} Selected Curves"), NumBufferedCurves, NumActiveCurves);
+	TAttribute<FText> ApplyBufferedCurvesText = FText::Format(LOCTEXT("ApplyBufferedCurvesContextMenu", "Apply {0} Buffered Curves onto {1} Selected Curves"), NumBufferedCurves, NumActiveCurves);
 
 	const FCurveModel* HoveredCurve = HoveredCurveID.IsSet() ? CurveEditor->FindCurve(HoveredCurveID.GetValue()) : nullptr;
 
