CurveEditor: Adjusted copy/paste logic for curves to better handle copying Control Rig controls between rigs. This now allows copying related controls between multiple rigs (ie: foot_fk_l, foot_fk_r, head_fk onto a matching set of controls on another object, but we also can support copying a single group (foot_fk_l onto foot_fk_r). Individual curves can still be copied to one or more target curves regardless of intent.

diff --git a/Engine/Plugins/Animation/ControlRig/Source/ControlRig/Private/Sequencer/MovieSceneControlRigParameterSection.cpp b/Engine/Plugins/Animation/ControlRig/Source/ControlRig/Private/Sequencer/MovieSceneControlRigParameterSection.cpp
--- a/Engine/Plugins/Animation/ControlRig/Source/ControlRig/Private/Sequencer/MovieSceneControlRigParameterSection.cpp
+++ b/Engine/Plugins/Animation/ControlRig/Source/ControlRig/Private/Sequencer/MovieSceneControlRigParameterSection.cpp
@@ -336,6 +336,8 @@ struct FParameterTransformChannelEditorData
 	{
 		ControlRig = InControlRig;
 		ParameterName = InName;
+		static FText LongIntentFormatStr = NSLOCTEXT("MovieSceneControlParameterRigSection", "LongIntentFormatString", "{GroupName}.{IntentName}");
+
 		//FText LocationGroup = NSLOCTEXT("MovieSceneControlParameterRigSection", "Location", "Location");
 		//FText RotationGroup = NSLOCTEXT("MovieSceneControlParameterRigSection", "Rotation", "Rotation");
 		//FText ScaleGroup = NSLOCTEXT("MovieSceneControlParameterRigSection", "Scale", "Scale");
@@ -349,6 +351,7 @@ struct FParameterTransformChannelEditorData
 			TotalName += ".Location.X";
 			MetaData[0].SetIdentifiers(FName(*TotalName), NSLOCTEXT("MovieSceneControlParameterRigSection", "Location.X", "Location.X"), TransformGroup);
 			MetaData[0].IntentName = NSLOCTEXT("MovieSceneControlParameterRigSection", "Location.X", "Location.X");
+			MetaData[0].LongIntentNameFormat = LongIntentFormatStr;
 			TotalName = NameAsString;
 
 			MetaData[0].bEnabled = bEnabledOverride && EnumHasAllFlags(Mask, EMovieSceneTransformChannel::TranslationX);
@@ -360,6 +363,7 @@ struct FParameterTransformChannelEditorData
 			TotalName += ".Location.Y";
 			MetaData[1].SetIdentifiers(FName(*TotalName), NSLOCTEXT("MovieSceneControlParameterRigSection", "Location.Y", "Location.Y"), TransformGroup);
 			MetaData[1].IntentName = NSLOCTEXT("MovieSceneControlParameterRigSection", "Location.Y", "Location.Y");
+			MetaData[1].LongIntentNameFormat = LongIntentFormatStr;
 			TotalName = NameAsString;
 
 			MetaData[1].bEnabled = bEnabledOverride && EnumHasAllFlags(Mask, EMovieSceneTransformChannel::TranslationY);
@@ -371,6 +375,7 @@ struct FParameterTransformChannelEditorData
 			TotalName += ".Location.Z";
 			MetaData[2].SetIdentifiers(FName(*TotalName), NSLOCTEXT("MovieSceneControlParameterRigSection", "Location.Z", "Location.Z"), TransformGroup);
 			MetaData[2].IntentName = NSLOCTEXT("MovieSceneControlParameterRigSection", "Location.Z", "Location.Z");
+			MetaData[2].LongIntentNameFormat = LongIntentFormatStr;
 			TotalName = NameAsString;
 
 			MetaData[2].bEnabled = bEnabledOverride && EnumHasAllFlags(Mask, EMovieSceneTransformChannel::TranslationZ);
@@ -383,6 +388,7 @@ struct FParameterTransformChannelEditorData
 			TotalName += ".Rotation.X";
 			MetaData[3].SetIdentifiers(FName(*TotalName), NSLOCTEXT("MovieSceneControlParameterRigSection", "Rotation.X", "Rotation.Roll"), TransformGroup);
 			MetaData[3].IntentName = NSLOCTEXT("MovieSceneControlParameterRigSection", "Rotation.X", "Rotation.Roll");
+			MetaData[3].LongIntentNameFormat = LongIntentFormatStr;
 			TotalName = NameAsString;
 
 			MetaData[3].bEnabled = bEnabledOverride && EnumHasAllFlags(Mask, EMovieSceneTransformChannel::RotationX);
@@ -394,6 +400,7 @@ struct FParameterTransformChannelEditorData
 			TotalName += ".Rotation.Y";
 			MetaData[4].SetIdentifiers(FName(*TotalName), NSLOCTEXT("MovieSceneControlParameterRigSection", "Rotation.Y", "Rotation.Pitch"), TransformGroup);
 			MetaData[4].IntentName = NSLOCTEXT("MovieSceneControlParameterRigSection", "Rotation.Y", "Rotation.Pitch");
+			MetaData[4].LongIntentNameFormat = LongIntentFormatStr;
 			TotalName = NameAsString;
 
 			MetaData[4].bEnabled = bEnabledOverride && EnumHasAllFlags(Mask, EMovieSceneTransformChannel::RotationY);
@@ -405,6 +412,7 @@ struct FParameterTransformChannelEditorData
 			TotalName += ".Rotation.Z";
 			MetaData[5].SetIdentifiers(FName(*TotalName), NSLOCTEXT("MovieSceneControlParameterRigSection", "Rotation.Z", "Rotation.Yaw"), TransformGroup);
 			MetaData[5].IntentName = NSLOCTEXT("MovieSceneControlParameterRigSection", "Rotation.Z", "Rotation.Yaw");
+			MetaData[5].LongIntentNameFormat = LongIntentFormatStr;
 			TotalName = NameAsString;
 
 			MetaData[5].bEnabled = bEnabledOverride && EnumHasAllFlags(Mask, EMovieSceneTransformChannel::RotationZ);
@@ -417,6 +425,7 @@ struct FParameterTransformChannelEditorData
 			TotalName += ".Scale.X";
 			MetaData[6].SetIdentifiers(FName(*TotalName), NSLOCTEXT("MovieSceneControlParameterRigSection", "Scale.X", "Scale.X"), TransformGroup);
 			MetaData[6].IntentName = NSLOCTEXT("MovieSceneControlParameterRigSection", "Scale.X", "Scale.X");
+			MetaData[6].LongIntentNameFormat = LongIntentFormatStr;
 			TotalName = NameAsString;
 
 			MetaData[6].bEnabled = bEnabledOverride && EnumHasAllFlags(Mask, EMovieSceneTransformChannel::ScaleX);
@@ -428,6 +437,7 @@ struct FParameterTransformChannelEditorData
 			TotalName += ".Scale.Y";
 			MetaData[7].SetIdentifiers(FName(*TotalName), NSLOCTEXT("MovieSceneControlParameterRigSection", "Scale.Y", "Scale.Y"), TransformGroup);
 			MetaData[7].IntentName = NSLOCTEXT("MovieSceneControlParameterRigSection", "Scale.Y", "Scale.Y");
+			MetaData[7].LongIntentNameFormat = LongIntentFormatStr;
 			TotalName = NameAsString;
 
 			MetaData[7].bEnabled = bEnabledOverride && EnumHasAllFlags(Mask, EMovieSceneTransformChannel::ScaleY);
@@ -439,6 +449,7 @@ struct FParameterTransformChannelEditorData
 			TotalName += ".Scale.Z";
 			MetaData[8].SetIdentifiers(FName(*TotalName), NSLOCTEXT("MovieSceneControlParameterRigSection", "Scale.Z", "Scale.Z"), TransformGroup);
 			MetaData[8].IntentName = NSLOCTEXT("MovieSceneControlParameterRigSection", "Scale.Z", "Scale.Z");
+			MetaData[8].LongIntentNameFormat = LongIntentFormatStr;
 			TotalName = NameAsString;
 
 			MetaData[8].bEnabled = bEnabledOverride && EnumHasAllFlags(Mask, EMovieSceneTransformChannel::ScaleZ);
diff --git a/Engine/Source/Editor/CurveEditor/Private/CurveEditor.cpp b/Engine/Source/Editor/CurveEditor/Private/CurveEditor.cpp
--- a/Engine/Source/Editor/CurveEditor/Private/CurveEditor.cpp
+++ b/Engine/Source/Editor/CurveEditor/Private/CurveEditor.cpp
@@ -375,6 +375,7 @@ bool FCurveEditor::ShouldAutoFrame() const
 	return Settings->GetAutoFrameCurveEditor();
 }
 
+
 void FCurveEditor::BindCommands()
 {
 	UCurveEditorSettings* CurveSettings = Settings;
@@ -385,7 +386,8 @@ void FCurveEditor::BindCommands()
 
 	CommandList->MapAction(FGenericCommands::Get().Cut, FExecuteAction::CreateSP(this, &FCurveEditor::CutSelection));
 	CommandList->MapAction(FGenericCommands::Get().Copy, FExecuteAction::CreateSP(this, &FCurveEditor::CopySelection));
-	CommandList->MapAction(FGenericCommands::Get().Paste, FExecuteAction::CreateSP(this, &FCurveEditor::PasteKeys, TSet<FCurveModelID>()));
+	CommandList->MapAction(FGenericCommands::Get().Paste, FExecuteAction::CreateSP(this, &FCurveEditor::PasteKeys, TSet<FCurveModelID>(), false));
+	CommandList->MapAction(FCurveEditorCommands::Get().PasteOverwriteRange, FExecuteAction::CreateSP(this, &FCurveEditor::PasteKeys, TSet<FCurveModelID>(), true));
 
 	CommandList->MapAction(FCurveEditorCommands::Get().ZoomToFit, FExecuteAction::CreateSP(this, &FCurveEditor::ZoomToFit, EAxisList::All));
 	CommandList->MapAction(FCurveEditorCommands::Get().ZoomToFitHorizontal, FExecuteAction::CreateSP(this, &FCurveEditor::ZoomToFit, EAxisList::X));
@@ -1259,6 +1261,7 @@ void FCurveEditor::CopySelection() const
 
 					CopyableCurveKeys->ShortDisplayName = Curve->GetShortDisplayName().ToString();
 					CopyableCurveKeys->LongDisplayName = Curve->GetLongDisplayName().ToString();
+					CopyableCurveKeys->LongIntentionName = Curve->GetLongIntentionName();
 					CopyableCurveKeys->IntentionName = Curve->GetIntentionName();
 					CopyableCurveKeys->KeyPositions.SetNum(NumKeys, false);
 					CopyableCurveKeys->KeyAttributes.SetNum(NumKeys, false);
@@ -1401,198 +1404,306 @@ void FCurveEditor::ImportCopyBufferFromText(const FString& TextToImport, /*out*/
 	TempPackage->RemoveFromRoot();
 }
 
-void FCurveEditor::PasteKeys(TSet<FCurveModelID> CurveModelIDs)
+TSet<FCurveModelID> FCurveEditor::GetTargetCurvesForPaste() const
 {
-	// Grab the text to paste from the clipboard
-	FString TextToImport;
-	FPlatformApplicationMisc::ClipboardPaste(TextToImport);
+	TSet<FCurveModelID> TargetCurves;
 
-	TArray<UCurveEditorCopyBuffer*> ImportedCopyBuffers;
-	ImportCopyBufferFromText(TextToImport, ImportedCopyBuffers);
-
-	if (ImportedCopyBuffers.Num() == 0)
+	TOptional<FCurveModelID> HoveredID;
+	if (WeakPanel.IsValid())
 	{
-		return;
+		for (TSharedPtr<SCurveEditorView> View : WeakPanel.Pin()->GetViews())
+		{
+			if (View.IsValid() && View->GetHoveredCurve().IsSet())
+			{
+				HoveredID = View->GetHoveredCurve().GetValue();
+				break;
+			}
+		}
 	}
 
-	// Determine whether all the copied keys are from the same curve, if yes, they can all be pasted to the target curves without name matching
-	bool bAllCopiedCurvesLongNameEqual = true;
-	FString AllCopiedCurvesLongName;
-	for (UCurveEditorCopyBuffer* CopyBuffer : ImportedCopyBuffers)
+	if (HoveredID.IsSet())
+	{
+		TargetCurves.Add(HoveredID.GetValue());
+	}
+	else
 	{
-		for (UCurveEditorCopyableCurveKeys* CopyableCurveKeys : CopyBuffer->Curves)
+		TArray<FCurveEditorTreeItemID> NodesToSearch;
+
+		// Try nodes with selected keys
+		for (const TTuple<FCurveModelID, FKeyHandleSet>& Pair : Selection.GetAll())
 		{
-			if (AllCopiedCurvesLongName.IsEmpty())
+			TargetCurves.Add(Pair.Key);
+		}
+
+		// Try selected nodes
+		if (TargetCurves.Num() == 0)
+		{
+			for (const TTuple<FCurveEditorTreeItemID, ECurveEditorTreeSelectionState>& Pair : GetTreeSelection())
 			{
-				AllCopiedCurvesLongName = CopyableCurveKeys->LongDisplayName;
+				NodesToSearch.Add(Pair.Key);
 			}
-			else if (CopyableCurveKeys->LongDisplayName != AllCopiedCurvesLongName)
+		}
+
+		for (const FCurveEditorTreeItemID& TreeItemID : NodesToSearch)
+		{
+			const FCurveEditorTreeItem& TreeItem = GetTreeItem(TreeItemID);
+			for (const FCurveModelID& CurveModelID : TreeItem.GetCurves())
 			{
-				bAllCopiedCurvesLongNameEqual = false;
-				break;
-			}			
+				TargetCurves.Add(CurveModelID);
+			}
 		}
 	}
 
-	bool bSelectionNeedsLongNames = false;
+	return TargetCurves;
+}
 
-	if (CurveModelIDs.Num() == 0)
+bool FCurveEditor::CopyBufferCurveToCurveID(const UCurveEditorCopyableCurveKeys* InSourceCurve, const FCurveModelID InTargetCurve, TOptional<double> InTimeOffset, const bool bInAddToSelection, const bool bInOverwriteRange)
+{
+	FCurveModel* TargetCurve = FindCurve(InTargetCurve);
+	if (!InSourceCurve || !TargetCurve)
+	{
+		return false;
+	}
+
+	// Sometimes when you paste you want to delete any keys that already exist in the timerange you'll be replacing
+	// because mixing the pasted results with the original results wouldn't make any sense.
+	if (bInOverwriteRange)
 	{
-		TOptional<FCurveModelID> HoveredID;
-		if (WeakPanel.IsValid())
+		TArray<FKeyHandle> KeysToRemove;
+		double MinKeyTime = TNumericLimits<double>::Max();
+		double MaxKeyTime = TNumericLimits<double>::Lowest(); 
+		for (int32 Index = 0; Index < InSourceCurve->KeyPositions.Num(); ++Index)
 		{
-			for (TSharedPtr<SCurveEditorView> View : WeakPanel.Pin()->GetViews())
+			FKeyPosition KeyPosition = InSourceCurve->KeyPositions[Index];
+			if (InTimeOffset.IsSet())
 			{
-				if (View.IsValid() && View->GetHoveredCurve().IsSet())
-				{
-					HoveredID = View->GetHoveredCurve().GetValue();
-					break;
-				}
+				KeyPosition.InputValue += InTimeOffset.GetValue();
+			}
+			if (KeyPosition.InputValue < MinKeyTime)
+			{
+				MinKeyTime = KeyPosition.InputValue;
+			}
+			if (KeyPosition.InputValue > MaxKeyTime)
+			{
+				MaxKeyTime = KeyPosition.InputValue;
 			}
 		}
-	
-		if (HoveredID.IsSet())
+
+		// Just double checking we actually set a Min/Max time so we don't wipe out every key to infinity.
+		if (InSourceCurve->KeyPositions.Num() > 0)
 		{
-			CurveModelIDs.Add(HoveredID.GetValue());
+			TargetCurve->GetKeys(*this, MinKeyTime, MaxKeyTime, TNumericLimits<double>::Lowest(), TNumericLimits<double>::Max(), KeysToRemove);
 		}
-		else
+
+		TargetCurve->RemoveKeys(KeysToRemove);
+	}
+
+	for (int32 Index = 0; Index < InSourceCurve->KeyPositions.Num(); ++Index)
+	{
+		FKeyPosition KeyPosition = InSourceCurve->KeyPositions[Index];
+		if (InTimeOffset.IsSet())
 		{
-			TArray<FCurveEditorTreeItemID> NodesToSearch;
+			KeyPosition.InputValue += InTimeOffset.GetValue();
+		}
 
-			// Try nodes with selected keys
-			for (const TTuple<FCurveModelID, FKeyHandleSet>& Pair : Selection.GetAll())
-			{
-				CurveModelIDs.Add(Pair.Key);
-			}
+		TOptional<FKeyHandle> KeyHandle = TargetCurve->AddKey(KeyPosition, InSourceCurve->KeyAttributes[Index]);
+		if (KeyHandle.IsSet() && bInAddToSelection)
+		{
+			Selection.Add(FCurvePointHandle(InTargetCurve, ECurvePointType::Key, KeyHandle.GetValue()));
+		}
+	}
 
-			// Try selected nodes
-			if (CurveModelIDs.Num() == 0)
-			{
-				for (const TTuple<FCurveEditorTreeItemID, ECurveEditorTreeSelectionState>& Pair : GetTreeSelection())
-				{
-					NodesToSearch.Add(Pair.Key);
-				}
+	return true;
+}
 
-				// If no curves are selected, paste to the entire tree using fully qualified long names
-				if (NodesToSearch.Num() == 0)
-				{
-					bSelectionNeedsLongNames = true;
-					bAllCopiedCurvesLongNameEqual = false;
-					if (Tree.GetAllItems().GetKeys(NodesToSearch) == 0)
-					{
-						// If we don't have any curves to paste in to, exit now
-						return;
-					}
-				}
-			}
+void FCurveEditor::PasteKeys(TSet<FCurveModelID> CurveModelIDs, const bool bInOverwriteRange)
+{
+	// Grab the text to paste from the clipboard
+	FString TextToImport;
+	FPlatformApplicationMisc::ClipboardPaste(TextToImport);
 
-			for (const FCurveEditorTreeItemID& TreeItemID: NodesToSearch)
-			{
-				FCurveEditorTreeItem& TreeItem = GetTreeItem(TreeItemID);
-				for (const FCurveModelID& CurveModelID : TreeItem.GetCurves())
-				{
-					CurveModelIDs.Add(CurveModelID);
-				}
-			}
-		}
+	TArray<UCurveEditorCopyBuffer*> ImportedCopyBuffers;
+	ImportCopyBufferFromText(TextToImport, ImportedCopyBuffers);
+
+	if (ImportedCopyBuffers.Num() == 0)
+	{
+		return;
 	}
+
+	// There are numerous scenarios that Copy/Paste needs to handle.
+	// 1:1				 - Copying a single curve to another single curve should always work.
+	// 1:Multiple		 - Copying a single curve with multiple target curves should always work, the value will just be written into each one.
+	// Multiple (Related): Multiple (Related) 
+	//					 - Copying multiple curves between related controls, ie: fk_foot_l and fk_foot_r from one rig to another.
+	//					 - If their long intent name matches, we consider them to be related controls. If their intent name doesn't match
+	//					 - then we consider them unrelated controls.
+	// Multiple (Unrelated):Multiple (Unrelated)
+	//					 - If the long name doesn't match then we fall back to just the intent name. We want to handle copying both from one
+	//					 - group of controls to multiple groups of controls, matching each by short intent name. This lets you copy fk_foot_l
+	//					 - onto fk_foot_r and fk_spine_1 at the same time. We also handle trying to copy from multiple groups of controls
+	//					 - onto multiple groups of controls - this falls back to a index-in-array order based copy and tries to ensure that
+	//					 - the intent for each one (ie: transform.x) copies onto the first target transform.x, and then the next source that
+	//					 - has a transform.x intent gets copied onto the *second* target transform.x.
+	// Multiple (Unrelated):1
+	//					 - This one is mostly an unhandled case and the last source intent will win on the target group, so fk_foot_l and fk_foot_r
+	// 					 - pasted onto fk_spine_1, fk_spine_1 will just get the intents from fk_foot_r and fk_foot_l is ignored. This order isn't
+	//					 - guranteed though because it's using the order the curves are in the internal arrays.
 	
-	if (CurveModelIDs.Num() == 0)
+	// There should only be one copy buffer, but the way the import works returns an array.
+	ensureMsgf(ImportedCopyBuffers.Num() == 1, TEXT("Multiple copy buffers pasted at one time, only the first one will be used!"));
+	UCurveEditorCopyBuffer* SourceBuffer = ImportedCopyBuffers[0];
+
+	// Figure out which CurveModelIDs we're trying to paste to. If they're not already specified, we try to find hovered curves,
+	// and failing that we try to find all curves.
+	TSet<FCurveModelID> TargetCurves = CurveModelIDs.Num() > 0 ? CurveModelIDs : GetTargetCurvesForPaste();
+	
+	if (TargetCurves.Num() == 0)
 	{
 		return;
 	}
 
-	FScopedTransaction Transaction(LOCTEXT("PasteKeys", "Paste Keys"));
+	// When we're pasting keys, we want the first key to paste where the timeslider is
+	TOptional<double> TimeOffset;
+	bool bApplyOffset = !SourceBuffer->bAbsolutePosition;
 
+	if (bApplyOffset)
+	{
+		TSharedPtr<ITimeSliderController> TimeSliderController = WeakTimeSliderController.Pin();
+		if (TimeSliderController.IsValid())
+		{
+			FFrameRate TickResolution = TimeSliderController->GetTickResolution();
+
+			TimeOffset = TimeSliderController->GetScrubPosition() / TickResolution;
+		}
+		else
+		{
+			TimeOffset = SourceBuffer->TimeOffset;
+		}
+	}
+
+	FScopedTransaction Transaction(LOCTEXT("PasteKeys", "Paste Keys"));
 	Selection.Clear();
 
-	// We don't expect/want multiple copy buffers, but the way serialization works it's a possibile edge case,
-	// so we'll try to handle it sanely and treat each one as an individual block to paste.
-	for (UCurveEditorCopyBuffer* CopyBuffer : ImportedCopyBuffers)
-	{
-		bool bUseLongDisplayName = bSelectionNeedsLongNames;
 
-		double TimeOffset = 0.0f;
-		bool bApplyOffset = !CopyBuffer->bAbsolutePosition;
+	// Two simple cases, 1 to 1 and 1 to many.
+	TArray<TPair<UCurveEditorCopyableCurveKeys*, FCurveModelID>> CopyPairs;
 
-		if (bApplyOffset)
+	if (SourceBuffer->Curves.Num() == 1)
+	{
+		for (FCurveModelID TargetCurveID : TargetCurves)
 		{
-			TSharedPtr<ITimeSliderController> TimeSliderController = WeakTimeSliderController.Pin();
-			if (TimeSliderController.IsValid())
-			{
-				FFrameRate TickResolution = TimeSliderController->GetTickResolution();
+			CopyPairs.Add(TPair<UCurveEditorCopyableCurveKeys*, FCurveModelID>(SourceBuffer->Curves[0], TargetCurveID));
+		}
+	}
+	else
+	{
+		// The more complicated is the Multiple:Multiple / Multiple:1 (which is really just the same). We want to
+		// prioritize matching up longer names if possible - this allows us to copy multiple controls to multiple
+		// controls, such as starting with fk_foot_l and fk_foot_r and pasting to fk_foot_l, fk_foot_r, fk_neck_01.
+		// We will match up the transform/scale/rotation for the fk_foot_l/fk_foot_r and don't touch fk_neck_01 in this
+		// example. If no matches are made, then we fall back to the shorter intent string - where we just copy
+		// transform.xyz to transform.xyz even though the source may be fk_foot_l and the target is fk_foot_r.
 
-				TimeOffset = TimeSliderController->GetScrubPosition() / TickResolution;
+		// If any of the long names match (ie: fk_foot_l.transform.x) then we'll use long name matching for all.
+		bool bUseLongNameForMatches = false;
+		for (const UCurveEditorCopyableCurveKeys* SourceCurveKeys : SourceBuffer->Curves)
+		{
+			for (const FCurveModelID& TargetCurveID : TargetCurves)
+			{
+				FCurveModel* TargetCurve = FindCurve(TargetCurveID);
+				if (TargetCurve)
+				{
+					if (SourceCurveKeys->LongIntentionName == TargetCurve->GetLongIntentionName())
+					{
+						bUseLongNameForMatches = true;
+						break;
+					}
+				}
 			}
-			else
+
+			// Exit out of the outer loop too if we've got a match.
+			if (bUseLongNameForMatches)
 			{
-				TimeOffset = CopyBuffer->TimeOffset;
+				break;
 			}
 		}
 
-		TArray<UCurveEditorCopyableCurveKeys*> UsedCurves;
+		// Multiple to Multiple curve copying can get complicated when we only have the short intent name to deal with it, so
+		// this creates an edge case where you're copying one set of intents (ie: transform.x, transform.y, transform.z) onto
+		// multiple objects with those intents... we want to support this, but we don't support copying from multiple objects
+		// onto multiple objects unless their LongIntentionName matches as it gets too confusing to match up.
+		bool bOnlyOneSetOfSourceIntentions = true;
+		{
+			TMap<FString, int32> IntentionUseCounts;
+			for (UCurveEditorCopyableCurveKeys* SourceCurveKeys : SourceBuffer->Curves)
+			{
+				IntentionUseCounts.FindOrAdd(SourceCurveKeys->IntentionName)++;
+			}
 
-		for (FCurveModelID CurveID : CurveModelIDs)
+			for (TPair<FString, int32>& Pair : IntentionUseCounts)
+			{
+				if (Pair.Value > 1)
+				{
+					bOnlyOneSetOfSourceIntentions = false;
+					break;
+				}
+			}
+		}
+		
+		TSet<FCurveModelID> CurvesToMatchTo = TargetCurves;
+		for (UCurveEditorCopyableCurveKeys* SourceCurveKeys : SourceBuffer->Curves)
 		{
-			FCurveModel* Curve = FindCurve(CurveID);
-			if (Curve)
+			TArray<FCurveModelID> CurvesToRemove;
+			for (const FCurveModelID& TargetCurveID : CurvesToMatchTo)
 			{
-				const FString CurveLongDisplayName = Curve->GetLongDisplayName().ToString();
-				const FString CurveIntentionName = Curve->GetIntentionName();
-				
-				bool bFoundMatch = false;
-				for (UCurveEditorCopyableCurveKeys* CopyableCurveKeys : CopyBuffer->Curves)
+				FCurveModel* TargetCurve = FindCurve(TargetCurveID);
+				if (TargetCurve)
 				{
-					// Use up all curves in the copied buffer before reusing. For example, if the following 6 curves have been copied,
-					//   Cube1.Location.X, Cube1.Location.Y, Cube1.Location.Z
-					//   Cube2.Location.X, Cube2.Location.Y, Cube2.Location.Z
-					// and the user attempts to paste onto multiple objects, Cube1's curves will match before Cube2's curves if this restriction is not in place.
-					//
-					if (UsedCurves.Contains(CopyableCurveKeys))
-					{
-						continue;
-					}
+					const bool bNameMatches = bUseLongNameForMatches ?
+						SourceCurveKeys->LongIntentionName == TargetCurve->GetLongIntentionName() :
+						SourceCurveKeys->IntentionName == TargetCurve->GetIntentionName();
 
-					if (bAllCopiedCurvesLongNameEqual ||
-						(!bUseLongDisplayName && CurveIntentionName.Equals(CopyableCurveKeys->IntentionName))
-						|| (bUseLongDisplayName && CurveLongDisplayName.Equals(CopyableCurveKeys->LongDisplayName)))
+					if (bNameMatches)
 					{
-						bFoundMatch = true;
-
-						for (int32 Index = 0; Index < CopyableCurveKeys->KeyPositions.Num(); ++Index)
-						{
-							FKeyPosition KeyPosition = CopyableCurveKeys->KeyPositions[Index];
-							if (bApplyOffset)
-							{
-								KeyPosition.InputValue += TimeOffset;
-							}
-
-							TOptional<FKeyHandle> KeyHandle = Curve->AddKey(KeyPosition, CopyableCurveKeys->KeyAttributes[Index]);
-							if (KeyHandle.IsSet())
-							{
-								Selection.Add(FCurvePointHandle(CurveID, ECurvePointType::Key, KeyHandle.GetValue()));
-							}
-						}
-
-						UsedCurves.Add(CopyableCurveKeys);
-
-						if (UsedCurves.Num() == CopyBuffer->Curves.Num())
+						CopyPairs.Add(TPair<UCurveEditorCopyableCurveKeys*, FCurveModelID>(SourceCurveKeys, TargetCurveID));
+
+						// Don't try to match to this curve again. This lets us try to handle the case where we have
+						// multiple source objects (fk_foot_l, fk_foot_r) trying to copy to unrelated objects (cube1, cube2).
+						// They will fail the LongDisplayName check but get the IntentionName check, but we need to remove
+						// cube1 after the first time we match it so that fk_foot_r has a chance to paste into cube2 instead of cube1.
+						CurvesToRemove.Add(TargetCurveID);
+
+						// If we're copying from one object with multiple curves (ie: fk_foot_l) but we have multiple destination
+						// objects, we loop through all of the target curves and apply them using the IntentionName matches check.
+						// This only happens when using short intention names (as it's the more vague logic case), and we only
+						// do this when you have multiple source curves, but only one of each kind. If you have multiple source
+						// curves with multiple copies of the same intention, then we only apply it once to the first curve
+						// who's intention matches and then remove it from the pool so that the next source with the same
+						// intention (such as the second foot in the above example) gets a chance to write to the second
+						// target curve with the same destination.
+						bool bCopyToMultipleDestCurves = bOnlyOneSetOfSourceIntentions && !bUseLongNameForMatches;
+						if (!bCopyToMultipleDestCurves)
 						{
-							UsedCurves.Empty();
+							break;
 						}
-
-						break; // Just paste one of the matching curves
 					}
 				}
+			}
 
-				if (!bFoundMatch)
-				{
-					UE_LOG(LogCurveEditor, Warning, TEXT("Failed to find matching curve to copy onto: %s"), *CurveLongDisplayName);		
-				}			
+			for (FCurveModelID Curve : CurvesToRemove)
+			{
+				CurvesToMatchTo.Remove(Curve);
 			}
 		}
 	}
 
+	// Now that we've calculated the source curve for each destination curve, copy them over.
+	for (const TPair<UCurveEditorCopyableCurveKeys*, FCurveModelID>& Pair : CopyPairs)
+	{
+		const bool bAddToSelection = true;
+		CopyBufferCurveToCurveID(Pair.Key, Pair.Value, TimeOffset, bAddToSelection, bInOverwriteRange);
+	}
+
 	if (ShouldAutoFrame())
 	{
 		ZoomToFitSelection();
diff --git a/Engine/Source/Editor/CurveEditor/Private/CurveEditorCommands.cpp b/Engine/Source/Editor/CurveEditor/Private/CurveEditorCommands.cpp
--- a/Engine/Source/Editor/CurveEditor/Private/CurveEditorCommands.cpp
+++ b/Engine/Source/Editor/CurveEditor/Private/CurveEditorCommands.cpp
@@ -60,7 +60,8 @@ void FCurveEditorCommands::RegisterCommands()
 	UI_COMMAND(ToggleShowCurveEditorCurveToolTips, "Curve Tool Tips", "Show a tool tip with name and values when hovering over a curve.", EUserInterfaceActionType::ToggleButton, FInputChord() );
 
 	UI_COMMAND(AddKeyHovered, "Add Key", "Add a new key to this curve at the current position.", EUserInterfaceActionType::Button, FInputChord(EKeys::MiddleMouseButton) );
-	UI_COMMAND(PasteKeysHovered, "Paste", "Paste clipboard contents", EUserInterfaceActionType::Button, FInputChord(EModifierKey::Control, EKeys::V) );
+	UI_COMMAND(PasteKeysHovered, "Paste", "Paste clipboard contents", EUserInterfaceActionType::Button, FInputChord(EModifierKey::Control, EKeys::V));
+	UI_COMMAND(PasteOverwriteRange, "Paste Overwrite", "Pastes the keys from the clipboard, overwriting any key in destination track between the first and last pasted keys.", EUserInterfaceActionType::Button, FInputChord(EModifierKey::Control | EModifierKey::Shift, EKeys::V));
 
 	UI_COMMAND(AddKeyToAllCurves, "Add Key", "Add a new key to all curves at the current time.", EUserInterfaceActionType::Button, FInputChord(EKeys::Enter) );
 
diff --git a/Engine/Source/Editor/CurveEditor/Private/CurveEditorContextMenu.cpp b/Engine/Source/Editor/CurveEditor/Private/CurveEditorContextMenu.cpp
--- a/Engine/Source/Editor/CurveEditor/Private/CurveEditorContextMenu.cpp
+++ b/Engine/Source/Editor/CurveEditor/Private/CurveEditorContextMenu.cpp
@@ -76,6 +76,7 @@ void FCurveEditorContextMenu::BuildMenu(FMenuBuilder& MenuBuilder, TSharedRef<FC
 				// Modify Data
 				MenuBuilder.AddMenuEntry(FGenericCommands::Get().Cut);
 				MenuBuilder.AddMenuEntry(FGenericCommands::Get().Paste);
+				MenuBuilder.AddMenuEntry(FCurveEditorCommands::Get().PasteOverwriteRange);
 				MenuBuilder.AddMenuEntry(FGenericCommands::Get().Copy);
 				MenuBuilder.AddMenuEntry(FGenericCommands::Get().Delete);
 
diff --git a/Engine/Source/Editor/CurveEditor/Private/CurveEditorCopyBuffer.h b/Engine/Source/Editor/CurveEditor/Private/CurveEditorCopyBuffer.h
--- a/Engine/Source/Editor/CurveEditor/Private/CurveEditorCopyBuffer.h
+++ b/Engine/Source/Editor/CurveEditor/Private/CurveEditorCopyBuffer.h
@@ -26,9 +26,16 @@ public:
 	UPROPERTY()
 	FString LongDisplayName;
 
-	/** This curve's intention (such as Transform.X or Scale.X). Used internally to match up curves when saving/restoring curves between different objects. */
+	/** This curve's short intention (such as Transform.X or Scale.X). Used internally to match up curves when saving/restoring curves between different objects. */
 	UPROPERTY()
 	FString IntentionName;
+
+	/** 
+	* This curve's long intention (such as foot_fk_l.Transform.X or foot_fk_r.Scale.X). Used internally to match up curves when saving/restoring curves between different objects.
+	* Long intention names have priority in copy/paste over short intention names, but we fall back to short intention if it's unclear what the user is trying to do.
+	*/
+	UPROPERTY()
+	FString LongIntentionName;
 };
 
 UCLASS()
diff --git a/Engine/Source/Editor/CurveEditor/Public/CurveEditor.h b/Engine/Source/Editor/CurveEditor/Public/CurveEditor.h
--- a/Engine/Source/Editor/CurveEditor/Public/CurveEditor.h
+++ b/Engine/Source/Editor/CurveEditor/Public/CurveEditor.h
@@ -47,6 +47,7 @@ class ITimeSliderController;
 class SCurveEditorPanel;
 class SCurveEditorView;
 class UCurveEditorCopyBuffer;
+class UCurveEditorCopyableCurveKeys;
 class UCurveEditorSettings;
 struct FCurveDrawParams;
 struct FCurveEditorInitParams;
@@ -467,6 +468,8 @@ public:
 
 protected:
 	void ImportCopyBufferFromText(const FString& TextToImport, /*out*/ TArray<UCurveEditorCopyBuffer*>& ImportedCopyBuffers) const;
+	TSet<FCurveModelID> GetTargetCurvesForPaste() const;
+	bool CopyBufferCurveToCurveID(const UCurveEditorCopyableCurveKeys* InSourceCurve, const FCurveModelID InTargetCurve, TOptional<double> InTimeOffset, const bool bInAddToSelection, const bool bInOverwriteRange);
 
 	void GetChildCurveModelIDs(const FCurveEditorTreeItemID TreeItemID, TSet<FCurveModelID>& CurveModelIDs) const;
 
@@ -474,7 +477,7 @@ public:
 	/**
 	 * Paste keys
 	 */
-	void PasteKeys(TSet<FCurveModelID> CurveModelIDs);
+	void PasteKeys(TSet<FCurveModelID> CurveModelIDs, const bool bInOverwriteRange = false);
 
 	/**
 	 * Delete the currently selected keys
diff --git a/Engine/Source/Editor/CurveEditor/Public/CurveEditorCommands.h b/Engine/Source/Editor/CurveEditor/Public/CurveEditorCommands.h
--- a/Engine/Source/Editor/CurveEditor/Public/CurveEditorCommands.h
+++ b/Engine/Source/Editor/CurveEditor/Public/CurveEditorCommands.h
@@ -79,6 +79,7 @@ public:
 
 	TSharedPtr<FUICommandInfo> AddKeyHovered;
 	TSharedPtr<FUICommandInfo> PasteKeysHovered;
+	TSharedPtr<FUICommandInfo> PasteOverwriteRange;
 
 	TSharedPtr<FUICommandInfo> AddKeyToAllCurves;
 
diff --git a/Engine/Source/Editor/CurveEditor/Public/CurveModel.h b/Engine/Source/Editor/CurveEditor/Public/CurveModel.h
--- a/Engine/Source/Editor/CurveEditor/Public/CurveModel.h
+++ b/Engine/Source/Editor/CurveEditor/Public/CurveModel.h
@@ -345,6 +345,16 @@ public:
 		IntentionName = InIntentionName;
 	}
 
+	FORCEINLINE void SetLongIntentionName(const FString& InIntentionName)
+	{
+		LongIntentionName = InIntentionName;
+	}
+
+	FORCEINLINE FString GetLongIntentionName() const
+	{
+		return LongIntentionName;
+	}
+
 	/**
 	 */
 	FORCEINLINE void SetColor(const FLinearColor& InColor)
@@ -384,8 +394,14 @@ protected:
 	/** This curve's long display name. Used in situations where the UI doesn't provide enough context about what the curve is otherwise (such as "Floor.Transform.X") */
 	FText LongDisplayName;
 
-	/** This curve's intention (such as Transform.X or Scale.X). Used internally to match up curves when saving/restoring curves between different objects. */
+	/** This curve's short intention (such as Transform.X or Scale.X). Used internally to match up curves when saving/restoring curves between different objects. */
 	FString IntentionName;
+	
+	/** 
+	* This curve's long intention (such as foot_fk_l.Transform.X or foot_fk_r.Scale.X). Used internally to match up curves when saving/restoring curves between different objects.
+	* Long intention names have priority in copy/paste over short intention names, but we fall back to short intention if it's unclear what the user is trying to do.
+	*/
+	FString LongIntentionName;
 
 	/** This curve's display color */
 	FLinearColor Color;
diff --git a/Engine/Source/Editor/Sequencer/Private/IKeyArea.cpp b/Engine/Source/Editor/Sequencer/Private/IKeyArea.cpp
--- a/Engine/Source/Editor/Sequencer/Private/IKeyArea.cpp
+++ b/Engine/Source/Editor/Sequencer/Private/IKeyArea.cpp
@@ -239,6 +239,33 @@ FText GetOwningObjectBindingName(UMovieSceneTrack* InTrack, ISequencer& InSequen
 	return FText();
 }
 
+namespace UE
+{
+namespace MovieScene
+{
+	void CleanUpCurveEditorFormatStringInline(FString& InStr)
+	{
+		for (int32 Index = 0; Index < InStr.GetCharArray().Num() - 1; Index++)
+		{
+			TCHAR CurrentChar = InStr.GetCharArray()[Index];
+			TCHAR NextChar = InStr.GetCharArray()[Index + 1];
+			if (CurrentChar == '.' && NextChar == '.')
+			{
+				const bool bAllowShrinking = false;
+				const int32 NumToRemove = 1;
+				InStr.RemoveAt(Index, NumToRemove, bAllowShrinking);
+				Index--;
+			}
+		}
+
+		// We can still end up with a "." at the end using the above logic
+		if (InStr.GetCharArray()[InStr.GetCharArray().Num() - 1] == '.')
+		{
+			InStr.RemoveAt(InStr.GetCharArray().Num() - 1);
+		}
+	}
+}}
+
 TUniquePtr<FCurveModel> IKeyArea::CreateCurveEditorModel(TSharedRef<ISequencer> InSequencer) const
 {
 	ISequencerChannelInterface* EditorInterface = FindChannelEditorInterface();
@@ -260,43 +287,23 @@ TUniquePtr<FCurveModel> IKeyArea::CreateCurveEditorModel(TSharedRef<ISequencer>
 				// This track might be inside an object binding and we'd like to prepend the object binding's name for more context.
 				ObjectBindingName = GetOwningObjectBindingName(OwningTrack, *InSequencer);
 			}
+			
+			// This is a FText (even though the end result is a FString) because all the incoming data sources are FText
+			// And it's probably cheaper to convert one FText->FString instead of 4 shorter ones.
+			FFormatNamedArguments FormatArgs;
 
 			// Not all tracks have all the information so we need to format it differently depending on how many are valid.
-			TArray<FText> ValidNames;
+			FormatArgs.Add(TEXT("ObjectBindingName"), ObjectBindingName);
+			FormatArgs.Add(TEXT("OwningTrackName"), OwningTrackName);
+			FormatArgs.Add(TEXT("GroupName"), ChannelHandle.GetMetaData()->Group);
+			FormatArgs.Add(TEXT("DisplayName"), DisplayText);
 
-			if (!ObjectBindingName.IsEmptyOrWhitespace())
-			{
-				ValidNames.Add(ObjectBindingName);
-			}
-			if (!OwningTrackName.IsEmptyOrWhitespace())
-			{
-				ValidNames.Add(OwningTrackName);
-			}
-			if (!ChannelHandle.GetMetaData()->Group.IsEmptyOrWhitespace())
-			{
-				ValidNames.Add(ChannelHandle.GetMetaData()->Group);
-			}
-			if (!DisplayText.IsEmptyOrWhitespace())
-			{
-				ValidNames.Add(DisplayText);
-			}
+			FText FormatText = NSLOCTEXT("SequencerIKeyArea", "CurveLongDisplayNameFormat", "{ObjectBindingName}.{OwningTrackName}.{GroupName}.{DisplayName}");
+			FString FormattedText = FText::Format(FormatText, FormatArgs).ToString();
 
-			// Now we loop through and string them together into one big format string.
-			FText LongDisplayNameFormatString;
-			for (int32 NameIndex = 0; NameIndex < ValidNames.Num(); NameIndex++)
-			{
-				const bool bLastEntry = NameIndex == ValidNames.Num() - 1;
-				if (!bLastEntry)
-				{
-					LongDisplayNameFormatString = FText::Format(NSLOCTEXT("SequencerIKeyArea", "CurveLongDisplayNameFormat", "{0}`{{1}`}."), LongDisplayNameFormatString, NameIndex);
-				}
-				else
-				{
-					LongDisplayNameFormatString = FText::Format(NSLOCTEXT("SequencerIKeyArea", "CurveLongDisplayNameFormatEnd", "{0}`{{1}`}"), LongDisplayNameFormatString, NameIndex);
-				}
-			}
+			UE::MovieScene::CleanUpCurveEditorFormatStringInline(FormattedText);
 
-			FText LongDisplayName = FText::Format(LongDisplayNameFormatString, FFormatOrderedArguments(ValidNames));
+			FText LongDisplayName = FText::FromString(FormattedText);
 			const FText ShortDisplayName = DisplayText;
 			FString IntentName = ChannelHandle.GetMetaData()->IntentName.ToString();
 			if (IntentName.IsEmpty())
@@ -304,9 +311,23 @@ TUniquePtr<FCurveModel> IKeyArea::CreateCurveEditorModel(TSharedRef<ISequencer>
 				IntentName = ChannelHandle.GetMetaData()->Group.IsEmptyOrWhitespace() ? DisplayText.ToString() : FString::Format(TEXT("{0}.{1}"), { ChannelHandle.GetMetaData()->Group.ToString(), DisplayText.ToString() });
 			}
 
+			FText LongIntentNameFormat = ChannelHandle.GetMetaData()->LongIntentNameFormat;
+			if (LongIntentNameFormat.IsEmpty())
+			{
+				LongIntentNameFormat = NSLOCTEXT("SequencerIKeyArea", "LongIntentNameFormat", "{ObjectBindingName}.{GroupName}.{DisplayName}");
+			}
+			
+			FormatArgs.Add(TEXT("IntentName"), FText::FromString(IntentName));
+
+
+			FString LongIntentName = FText::Format(LongIntentNameFormat, FormatArgs).ToString();
+			UE::MovieScene::CleanUpCurveEditorFormatStringInline(LongIntentName);
+
+
 			CurveModel->SetShortDisplayName(DisplayText);
 			CurveModel->SetLongDisplayName(LongDisplayName);
 			CurveModel->SetIntentionName(IntentName);
+			CurveModel->SetLongIntentionName(LongIntentName);
 			if (Color.IsSet())
 			{
 				CurveModel->SetColor(Color.GetValue());
@@ -343,4 +364,4 @@ bool IKeyArea::ShouldShowCurve() const
 	}
 
 	return false;
-}
\ No newline at end of file
+}
diff --git a/Engine/Source/Runtime/MovieScene/Public/Channels/MovieSceneChannelEditorData.h b/Engine/Source/Runtime/MovieScene/Public/Channels/MovieSceneChannelEditorData.h
--- a/Engine/Source/Runtime/MovieScene/Public/Channels/MovieSceneChannelEditorData.h
+++ b/Engine/Source/Runtime/MovieScene/Public/Channels/MovieSceneChannelEditorData.h
@@ -63,6 +63,8 @@ struct FMovieSceneChannelMetaData
 	FText Group;
 	/** Intent name */
 	FText IntentName;
+	/* Optional. If unspecified IKeyArea::CreateCurveEditorModel will create a fallback. */
+	FText LongIntentNameFormat;
 	/** Optional color to draw underneath the keys on this channel */
 	TOptional<FLinearColor> Color;
 };
