Sequencer: Add commands to clear selection in Sequencer and Curve Editor (Esc)

diff --git a/Engine/Source/Editor/CurveEditor/Private/CurveEditor.cpp b/Engine/Source/Editor/CurveEditor/Private/CurveEditor.cpp
--- a/Engine/Source/Editor/CurveEditor/Private/CurveEditor.cpp
+++ b/Engine/Source/Editor/CurveEditor/Private/CurveEditor.cpp
@@ -414,6 +414,7 @@ void FCurveEditor::BindCommands()
 	CommandList->MapAction(FCurveEditorCommands::Get().SelectAllKeys, FExecuteAction::CreateSP(this, &FCurveEditor::SelectAllKeys));
 	CommandList->MapAction(FCurveEditorCommands::Get().SelectForward, FExecuteAction::CreateSP(this, &FCurveEditor::SelectForward));
 	CommandList->MapAction(FCurveEditorCommands::Get().SelectBackward, FExecuteAction::CreateSP(this, &FCurveEditor::SelectBackward));
+	CommandList->MapAction(FCurveEditorCommands::Get().SelectNone, FExecuteAction::CreateSP(this, &FCurveEditor::SelectNone));
 
 	{
 		FExecuteAction   ToggleInputSnapping     = FExecuteAction::CreateSP(this,   &FCurveEditor::ToggleInputSnapping);
@@ -1132,6 +1133,12 @@ void FCurveEditor::SelectBackward()
 	}
 }
 
+void FCurveEditor::SelectNone()
+{
+	Selection.Clear();
+}
+
+
 bool FCurveEditor::IsInputSnappingEnabled() const
 {
 	return InputSnapEnabledAttribute.Get();
diff --git a/Engine/Source/Editor/CurveEditor/Private/CurveEditorCommands.cpp b/Engine/Source/Editor/CurveEditor/Private/CurveEditorCommands.cpp
--- a/Engine/Source/Editor/CurveEditor/Private/CurveEditorCommands.cpp
+++ b/Engine/Source/Editor/CurveEditor/Private/CurveEditorCommands.cpp
@@ -87,6 +87,8 @@ void FCurveEditorCommands::RegisterCommands()
 	UI_COMMAND(SelectForward, "Select All Keys Forward", "Select all keys forward from the current time", EUserInterfaceActionType::Button, FInputChord(EModifierKey::Control, EKeys::RightBracket));
 	UI_COMMAND(SelectBackward, "Select All Keys Backward", "Select all keys backward from the current time", EUserInterfaceActionType::Button, FInputChord(EModifierKey::Control, EKeys::LeftBracket));
 
+	UI_COMMAND(SelectNone, "Select None", "Select none", EUserInterfaceActionType::Button, FInputChord(EKeys::Escape));
+
 	// Buffer and Apply Curves. Like copy and paste, but directly onto the curves they were stored from.
 	// These names are overwritten in CurveEditorContextMenu to show the number of stashed curves.
 	UI_COMMAND(BufferVisibleCurves, "Buffer Curves", "Stores a copy of the selected curves which can be applied back onto themselves.", EUserInterfaceActionType::Button, FInputChord());
diff --git a/Engine/Source/Editor/CurveEditor/Public/CurveEditor.h b/Engine/Source/Editor/CurveEditor/Public/CurveEditor.h
--- a/Engine/Source/Editor/CurveEditor/Public/CurveEditor.h
+++ b/Engine/Source/Editor/CurveEditor/Public/CurveEditor.h
@@ -226,6 +226,7 @@ public:
 	void SelectAllKeys();
 	void SelectForward();
 	void SelectBackward();
+	void SelectNone();
 
 	/** Toggle the expansion state of the selected nodes or all nodes if none selected */
 	void ToggleExpandCollapseNodes(bool bRecursive);
diff --git a/Engine/Source/Editor/CurveEditor/Public/CurveEditorCommands.h b/Engine/Source/Editor/CurveEditor/Public/CurveEditorCommands.h
--- a/Engine/Source/Editor/CurveEditor/Public/CurveEditorCommands.h
+++ b/Engine/Source/Editor/CurveEditor/Public/CurveEditorCommands.h
@@ -94,6 +94,7 @@ public:
 
 	TSharedPtr<FUICommandInfo> SelectForward;
 	TSharedPtr<FUICommandInfo> SelectBackward;
+	TSharedPtr<FUICommandInfo> SelectNone;
 
 	TSharedPtr<FUICommandInfo> BufferVisibleCurves;
 	TSharedPtr<FUICommandInfo> SwapBufferedCurves;
diff --git a/Engine/Source/Editor/Sequencer/Private/Sequencer.cpp b/Engine/Source/Editor/Sequencer/Private/Sequencer.cpp
--- a/Engine/Source/Editor/Sequencer/Private/Sequencer.cpp
+++ b/Engine/Source/Editor/Sequencer/Private/Sequencer.cpp
@@ -11246,6 +11246,10 @@ void FSequencer::BindCommands()
 		Commands.SelectBackward,
 		FExecuteAction::CreateSP(this, &FSequencer::SelectBackward));
 
+	SequencerCommandBindings->MapAction(
+		Commands.SelectNone,
+		FExecuteAction::CreateSP(this, &FSequencer::EmptySelection));
+
 	SequencerCommandBindings->MapAction(
 		Commands.StepToNextShot,
 		FExecuteAction::CreateSP( this, &FSequencer::StepToNextShot ) );
diff --git a/Engine/Source/Editor/Sequencer/Private/SequencerCommands.cpp b/Engine/Source/Editor/Sequencer/Private/SequencerCommands.cpp
--- a/Engine/Source/Editor/Sequencer/Private/SequencerCommands.cpp
+++ b/Engine/Source/Editor/Sequencer/Private/SequencerCommands.cpp
@@ -73,6 +73,7 @@ void FSequencerCommands::RegisterCommands()
 
 	UI_COMMAND( SelectForward, "Select All Keys and Sections Forward", "Select all keys and sections forward from the current time", EUserInterfaceActionType::Button, FInputChord(EModifierKey::Control, EKeys::RightBracket) );
 	UI_COMMAND( SelectBackward, "Select All Keys and Sections Backward", "Select all keys and sections backward from the current time", EUserInterfaceActionType::Button, FInputChord(EModifierKey::Control, EKeys::LeftBracket) );
+	UI_COMMAND( SelectNone, "Select None", "Select none", EUserInterfaceActionType::Button, FInputChord(EKeys::Escape));
 
 	UI_COMMAND( AddActorsToSequencer, "Add Actors", "Add actors to sequencer", EUserInterfaceActionType::Button, FInputChord(EModifierKey::Control, EKeys::A) );
 	UI_COMMAND( SetKey, "Set Key", "Sets a key on the selected tracks", EUserInterfaceActionType::Button, FInputChord(EKeys::Enter) );
diff --git a/Engine/Source/Editor/Sequencer/Public/SequencerCommands.h b/Engine/Source/Editor/Sequencer/Public/SequencerCommands.h
--- a/Engine/Source/Editor/Sequencer/Public/SequencerCommands.h
+++ b/Engine/Source/Editor/Sequencer/Public/SequencerCommands.h
@@ -193,6 +193,9 @@ public:
 	/** Select all keys and sections backward from the current time */
 	TSharedPtr< FUICommandInfo > SelectBackward;
 
+	/** Select none */
+	TSharedPtr< FUICommandInfo > SelectNone;
+
 	/** Add selected actors to sequencer */
 	TSharedPtr< FUICommandInfo > AddActorsToSequencer;
 
