Curve Editor: Save the tree view width percentage in curve settings so that the split between the tree view and the curve area is saved/restored.

#rnx

diff --git a/Engine/Plugins/FX/Niagara/Source/NiagaraEditorWidgets/Private/DetailCustomizations/NiagaraDataInterfaceCurveDetails.cpp b/Engine/Plugins/FX/Niagara/Source/NiagaraEditorWidgets/Private/DetailCustomizations/NiagaraDataInterfaceCurveDetails.cpp
--- a/Engine/Plugins/FX/Niagara/Source/NiagaraEditorWidgets/Private/DetailCustomizations/NiagaraDataInterfaceCurveDetails.cpp
+++ b/Engine/Plugins/FX/Niagara/Source/NiagaraEditorWidgets/Private/DetailCustomizations/NiagaraDataInterfaceCurveDetails.cpp
@@ -965,8 +965,6 @@ public:
 		TSharedPtr<SCurveEditorTree> CurveEditorTree;
 		CurveEditorPanel = SNew(SCurveEditorPanel, CurveEditor.ToSharedRef())
 			.MinimumViewPanelHeight(50.0f)
-			.TreeSplitterWidth(0.2f)
-			.ContentSplitterWidth(0.8f)
 			.TreeContent()
 			[
 				CurveProperties.Num() == 1
diff --git a/Engine/Source/Editor/CurveEditor/Private/CurveEditorSettings.cpp b/Engine/Source/Editor/CurveEditor/Private/CurveEditorSettings.cpp
--- a/Engine/Source/Editor/CurveEditor/Private/CurveEditorSettings.cpp
+++ b/Engine/Source/Editor/CurveEditor/Private/CurveEditorSettings.cpp
@@ -18,6 +18,8 @@ UCurveEditorSettings::UCurveEditorSettings()
 	SelectionColor = FLinearColor::White;
 	ParentSpaceCustomColor = FLinearColor(.93, .31, .19); //pastel orange
 	WorldSpaceCustomColor = FLinearColor(.198, .610, .558); //pastel teal
+
+	TreeViewWidth = 0.3f;
 }
 
 bool UCurveEditorSettings::GetAutoFrameCurveEditor() const
@@ -295,6 +297,15 @@ FLinearColor UCurveEditorSettings::GetNextRandomColor()
 	return Color;
 }
 
+void UCurveEditorSettings::SetTreeViewWidth(float InTreeViewWidth)
+{
+	if (InTreeViewWidth != TreeViewWidth)
+	{
+		TreeViewWidth = InTreeViewWidth;
+		SaveConfig();
+	}
+}
+
 void UCurveEditorSettings::PostEditChangeProperty(FPropertyChangedEvent& PropertyChangedEvent)
 {
 	FName PropertyName = (PropertyChangedEvent.Property != NULL) ? PropertyChangedEvent.Property->GetFName() : NAME_None;
diff --git a/Engine/Source/Editor/CurveEditor/Private/SCurveEditorPanel.cpp b/Engine/Source/Editor/CurveEditor/Private/SCurveEditorPanel.cpp
--- a/Engine/Source/Editor/CurveEditor/Private/SCurveEditorPanel.cpp
+++ b/Engine/Source/Editor/CurveEditor/Private/SCurveEditorPanel.cpp
@@ -16,6 +16,7 @@
 #include "CurveModel.h"
 #include "Curves/KeyHandle.h"
 #include "Delegates/Delegate.h"
+#include "Editor.h"
 #include "Filters/CurveEditorBakeFilter.h"
 #include "Filters/CurveEditorReduceFilter.h"
 #include "Filters/SCurveEditorFilterPanel.h"
@@ -181,6 +182,7 @@ void SCurveEditorPanel::Construct(const FArguments& InArgs, TSharedRef<FCurveEdi
 	WeakTabManager = InArgs._TabManager;
 
 	CachedSelectionSerialNumber = 0;
+
 	CurveEditor = InCurveEditor;
 
 	CurveEditor->SetPanel(SharedThis(this));
@@ -195,6 +197,21 @@ void SCurveEditorPanel::Construct(const FArguments& InArgs, TSharedRef<FCurveEdi
 
 	BindCommands();
 
+	ColumnFillCoefficients[0] = 0.3f;
+	ColumnFillCoefficients[1] = 0.7f;
+
+	if (CurveEditor->GetSettings())
+	{
+		ColumnFillCoefficients[0] = CurveEditor->GetSettings()->GetTreeViewWidth();
+		ColumnFillCoefficients[1] = 1.f - CurveEditor->GetSettings()->GetTreeViewWidth();
+	}
+
+	TAttribute<float> FillCoefficient_0, FillCoefficient_1;
+	{
+		FillCoefficient_0.Bind(TAttribute<float>::FGetter::CreateSP(this, &SCurveEditorPanel::GetColumnFillCoefficient, 0));
+		FillCoefficient_1.Bind(TAttribute<float>::FGetter::CreateSP(this, &SCurveEditorPanel::GetColumnFillCoefficient, 1));
+	}
+
 	// Create some Widgets
 	ISequencerWidgetsModule& SequencerWidgets = FModuleManager::Get().LoadModuleChecked<ISequencerWidgetsModule>("SequencerWidgets");
 	TSharedPtr<SWidget> TopTimeSlider = SNullWidget::NullWidget;
@@ -275,19 +292,22 @@ void SCurveEditorPanel::Construct(const FArguments& InArgs, TSharedRef<FCurveEdi
 	{
 		ChildSlot
 		[
-			SNew(SSplitter)
+			SAssignNew(TreeViewSplitter, SSplitter)
 			.Orientation(Orient_Horizontal)
 			.Style(FAppStyle::Get(), "SplitterDark")
 			.PhysicalSplitterHandleSize(3.0f)
+			.OnSplitterFinishedResizing(this, &SCurveEditorPanel::OnSplitterFinishedResizing)
 
 			+ SSplitter::Slot()
-			.Value(InArgs._TreeSplitterWidth)
+			.Value(FillCoefficient_0)
+			.OnSlotResized(SSplitter::FOnSlotResized::CreateSP(this, &SCurveEditorPanel::OnColumnFillCoefficientChanged, 0))
 			[
 				InArgs._TreeContent.Widget
 			]
 
 			+ SSplitter::Slot()
-			.Value(InArgs._ContentSplitterWidth)
+			.Value(FillCoefficient_1)
+			.OnSlotResized(SSplitter::FOnSlotResized::CreateSP(this, &SCurveEditorPanel::OnColumnFillCoefficientChanged, 1))
 			[
 				MainContent
 			]
@@ -1366,4 +1386,21 @@ EVisibility SCurveEditorPanel::ShouldInstructionOverlayBeVisible() const
 	const bool bCurvesAreVisible = CurveEditor->GetTreeSelection().Num() > 0 || CurveEditor->GetPinnedCurves().Num() > 0;
 	return bCurvesAreVisible ? EVisibility::Hidden : EVisibility::HitTestInvisible;
 }
+
+void SCurveEditorPanel::OnSplitterFinishedResizing()
+{
+	SSplitter::FSlot const& LeftSplitterSlot = TreeViewSplitter->SlotAt(0);
+	SSplitter::FSlot const& RightSplitterSlot = TreeViewSplitter->SlotAt(1);
+
+	OnColumnFillCoefficientChanged(LeftSplitterSlot.GetSizeValue(), 0);
+	OnColumnFillCoefficientChanged(RightSplitterSlot.GetSizeValue(), 1);
+
+	CurveEditor->GetSettings()->SetTreeViewWidth(LeftSplitterSlot.GetSizeValue());
+}
+
+void SCurveEditorPanel::OnColumnFillCoefficientChanged(float FillCoefficient, int32 ColumnIndex)
+{
+	ColumnFillCoefficients[ColumnIndex] = FillCoefficient;
+}
+
 #undef LOCTEXT_NAMESPACE
\ No newline at end of file
diff --git a/Engine/Source/Editor/CurveEditor/Public/CurveEditorSettings.h b/Engine/Source/Editor/CurveEditor/Public/CurveEditorSettings.h
--- a/Engine/Source/Editor/CurveEditor/Public/CurveEditorSettings.h
+++ b/Engine/Source/Editor/CurveEditor/Public/CurveEditorSettings.h
@@ -149,6 +149,11 @@ public:
 	/** Helper function to get next random linear color*/
 	static FLinearColor GetNextRandomColor();
 
+	/** Gets the tree view width percentage */
+	float GetTreeViewWidth() const { return TreeViewWidth; }
+	/** Sets the tree view width percentage */
+	void SetTreeViewWidth(float InTreeViewWidth);
+
 protected:
 
 	virtual void PostEditChangeProperty(FPropertyChangedEvent& PropertyChangedEvent) override;
@@ -194,6 +199,9 @@ protected:
 	UPROPERTY(config, EditAnywhere, Category = "Curve Editor")
 	TArray<FCustomColorForSpaceSwitch> ControlSpaceCustomColors;
 
+	UPROPERTY(config, EditAnywhere, Category = "Curve Editor")
+	float TreeViewWidth;
+
 private:
 	FOnCustomColorsChanged OnCustomColorsChangedEvent;
 };
diff --git a/Engine/Source/Editor/CurveEditor/Public/SCurveEditorPanel.h b/Engine/Source/Editor/CurveEditor/Public/SCurveEditorPanel.h
--- a/Engine/Source/Editor/CurveEditor/Public/SCurveEditorPanel.h
+++ b/Engine/Source/Editor/CurveEditor/Public/SCurveEditorPanel.h
@@ -53,8 +53,6 @@ class CURVEEDITOR_API SCurveEditorPanel : public SCompoundWidget
 {
 	SLATE_BEGIN_ARGS(SCurveEditorPanel)
 		: _GridLineTint(FLinearColor(0.1f, 0.1f, 0.1f, 1.f))
-		, _TreeSplitterWidth(0.3f)
-		, _ContentSplitterWidth(0.7f)
 		, _MinimumViewPanelHeight(300.0f)
 	{}
 
@@ -73,12 +71,6 @@ class CURVEEDITOR_API SCurveEditorPanel : public SCompoundWidget
 		/** Widget slot for the tree content */
 		SLATE_NAMED_SLOT(FArguments, TreeContent)
 
-		/** The width of the splitter slot for the tree */
-		SLATE_ARGUMENT(float, TreeSplitterWidth)
-
-		/** The width of the splitter slot for the main content */
-		SLATE_ARGUMENT(float, ContentSplitterWidth)
-
 		/** The minimum height for the panel which contains the curve editor views. */
 		SLATE_ARGUMENT(float, MinimumViewPanelHeight)
 
@@ -264,6 +256,17 @@ private:
 	/** Get a reference to the curve editor this panel represents. */
 	TSharedPtr<FCurveEditor> GetCurveEditor() const { return CurveEditor; }
 
+	float GetColumnFillCoefficient(int32 ColumnIndex) const
+	{
+		ensure(ColumnIndex == 0 || ColumnIndex == 1);
+		return ColumnFillCoefficients[ColumnIndex];
+	}
+
+	/** Called when a column fill percentage is changed by a splitter slot. */
+	void OnColumnFillCoefficientChanged(float FillCoefficient, int32 ColumnIndex);
+
+	void OnSplitterFinishedResizing();
+
 private:
 
 	/**
@@ -356,6 +359,11 @@ private:
 	/** A copy of the View Geometry used to represent the View portion of the Curve Editor. */
 	FGeometry CachedViewGeometry;
 
+	/** The fill coefficients of each column in the grid. */
+	float ColumnFillCoefficients[2];
+
+	TSharedPtr<class SSplitter> TreeViewSplitter;
+
 	/** Container of objects that are being used to edit keys on the curve editor */
 	TUniquePtr<FCurveEditorEditObjectContainer> EditObjects;
 
diff --git a/Engine/Source/Editor/Persona/Private/SAnimSequenceCurveEditor.cpp b/Engine/Source/Editor/Persona/Private/SAnimSequenceCurveEditor.cpp
--- a/Engine/Source/Editor/Persona/Private/SAnimSequenceCurveEditor.cpp
+++ b/Engine/Source/Editor/Persona/Private/SAnimSequenceCurveEditor.cpp
@@ -400,8 +400,6 @@ void SAnimSequenceCurveEditor::Construct(const FArguments& InArgs, const TShared
 		.GridLineTint(FLinearColor(0.f, 0.f, 0.f, 0.3f))
 		.ExternalTimeSliderController(InArgs._ExternalTimeSliderController)
 		.TabManager(InArgs._TabManager)
-		.TreeSplitterWidth(0.2f)
-		.ContentSplitterWidth(0.8f)
 		.TreeContent()
 		[
 			SNew(SVerticalBox)
