Curve Editor: Fix crash when removing keys on mouse up - ensure that the key handles are unique. Also, make sure last dragged keys are up to date on OnEndDrag.

diff --git a/Engine/Source/Editor/CurveEditor/Private/DragOperations/CurveEditorDragOperation_MoveKeys.cpp b/Engine/Source/Editor/CurveEditor/Private/DragOperations/CurveEditorDragOperation_MoveKeys.cpp
--- a/Engine/Source/Editor/CurveEditor/Private/DragOperations/CurveEditorDragOperation_MoveKeys.cpp
+++ b/Engine/Source/Editor/CurveEditor/Private/DragOperations/CurveEditorDragOperation_MoveKeys.cpp
@@ -130,7 +130,9 @@ void FCurveEditorDragOperation_MoveKeys::OnDrag(FVector2D InitialPosition, FVect
 		}
 
 		Curve->SetKeyPositions(KeyData.Handles, NewKeyPositionScratch, EPropertyChangeType::Interactive);
-		KeyData.LastDraggedKeyPositions = NewKeyPositionScratch;
+
+		// Make sure the last dragged key positions are up to date
+		Curve->GetKeyPositions(KeyData.Handles, KeyData.LastDraggedKeyPositions);
 	}
 }
 
@@ -196,7 +198,7 @@ void FCurveEditorDragOperation_MoveKeys::OnEndDrag(FVector2D InitialPosition, FV
 					{
 						if (KeyInRange != MaxKeyHandle)
 						{
-							KeysToRemove.Add(KeyInRange);
+							KeysToRemove.AddUnique(KeyInRange);
 						}
 					}
 				}
