Curve Editor: Performance: Add Curve caching so we only recreate the curves when something changes, giving us 30ms back in the specified jira test case.

Channels keep track if their owner's(MovieSceneSection) signature change, which we also now invalidate in the FCurveModel::Modify function.

InteractiveCurveView on Tick checks and sets a set of values to see if the view or something else has changed and then does a global recreate or checks each curve which then checks that new signature check to see if any individual curve was modified.

diff --git a/Engine/Plugins/Animation/ControlRig/Source/ControlRigEditor/Private/ControlRigSpaceChannelCurveModel.cpp b/Engine/Plugins/Animation/ControlRig/Source/ControlRigEditor/Private/ControlRigSpaceChannelCurveModel.cpp
--- a/Engine/Plugins/Animation/ControlRig/Source/ControlRigEditor/Private/ControlRigSpaceChannelCurveModel.cpp
+++ b/Engine/Plugins/Animation/ControlRig/Source/ControlRigEditor/Private/ControlRigSpaceChannelCurveModel.cpp
@@ -31,12 +31,32 @@ FControlRigSpaceChannelCurveModel::FControlRigSpaceChannelCurveModel(TMovieScene
 	}
 
 	WeakSection = OwningSection;
+	LastSignature = OwningSection->GetSignature();
 	WeakSequencer = InWeakSequencer;
 	SupportedViews = ViewID;
 
 	Color = STrackAreaView::BlendDefaultTrackColor(OwningSection->GetTypedOuter<UMovieSceneTrack>()->GetColorTint());
 }
 
+bool FControlRigSpaceChannelCurveModel::HasChangedAndResetTest()
+{
+	if (UMovieSceneSection* Section = WeakSection.Get())
+	{
+		if (Section->GetSignature() != LastSignature)
+		{
+			LastSignature = Section->GetSignature();
+			return true;
+		}
+		return false;
+	}
+	return true;
+}
+
+UObject* FControlRigSpaceChannelCurveModel::GetOwningObject() const 
+{
+	return WeakSection.Get();
+}
+
 FControlRigSpaceChannelCurveModel::~FControlRigSpaceChannelCurveModel()
 {
 	if (FMovieSceneChannelProxy* ChannelProxy = ChannelHandle.GetChannelProxy())
diff --git a/Engine/Plugins/Animation/ControlRig/Source/ControlRigEditor/Public/ControlRigSpaceChannelCurveModel.h b/Engine/Plugins/Animation/ControlRig/Source/ControlRigEditor/Public/ControlRigSpaceChannelCurveModel.h
--- a/Engine/Plugins/Animation/ControlRig/Source/ControlRigEditor/Public/ControlRigSpaceChannelCurveModel.h
+++ b/Engine/Plugins/Animation/ControlRig/Source/ControlRigEditor/Public/ControlRigSpaceChannelCurveModel.h
@@ -46,6 +46,9 @@ public:
 	virtual void RemoveKeys(TArrayView<const FKeyHandle> InKeys) override;
 
 	virtual void CreateKeyProxies(TArrayView<const FKeyHandle> InKeyHandles, TArrayView<UObject*> OutObjects) override;
+	
+	virtual UObject* GetOwningObject() const override;
+	virtual bool HasChangedAndResetTest() override;
 
 	//FKeyBarCurveModel
 	virtual void BuildContextMenu(const FCurveEditor& CurveEditor,FMenuBuilder& MenuBuilder, TOptional<FCurvePointHandle> ClickedPoint);
@@ -58,5 +61,6 @@ private:
 	TMovieSceneChannelHandle<FMovieSceneControlRigSpaceChannel> ChannelHandle;
 	TWeakObjectPtr<UMovieSceneSection> WeakSection;
 	TWeakPtr<ISequencer> WeakSequencer;
+	FGuid LastSignature;
 	FDelegateHandle OnDestroyHandle;
 };
\ No newline at end of file
diff --git a/Engine/Source/Editor/CurveEditor/Private/CurveModel.cpp b/Engine/Source/Editor/CurveEditor/Private/CurveModel.cpp
--- a/Engine/Source/Editor/CurveEditor/Private/CurveModel.cpp
+++ b/Engine/Source/Editor/CurveEditor/Private/CurveModel.cpp
@@ -18,6 +18,7 @@ void FCurveModel::SetKeyAttributes(TArrayView<const FKeyHandle> InKeys, const FK
 	}
 
 	SetKeyAttributes(InKeys, ExpandedAttributes);
+	Modify();
 }
 
 TOptional<FKeyHandle> FCurveModel::AddKey(const FKeyPosition& NewKeyPosition, const FKeyAttributes& InAttributes)
@@ -28,6 +29,5 @@ TOptional<FKeyHandle> FCurveModel::AddKey(const FKeyPosition& NewKeyPosition, co
 
 	TArrayView<TOptional<FKeyHandle>> Handles = MakeArrayView(&Handle, 1);
 	AddKeys(TArrayView<const FKeyPosition>(&NewKeyPosition, 1), TArrayView<const FKeyAttributes>(&InAttributes, 1), &Handles);
-
 	return Handle;
 }
diff --git a/Engine/Source/Editor/CurveEditor/Private/SCurveEditorView.cpp b/Engine/Source/Editor/CurveEditor/Private/SCurveEditorView.cpp
--- a/Engine/Source/Editor/CurveEditor/Private/SCurveEditorView.cpp
+++ b/Engine/Source/Editor/CurveEditor/Private/SCurveEditorView.cpp
@@ -21,6 +21,7 @@
 #include "SCurveEditorPanel.h"
 #include "Templates/Tuple.h"
 #include "Templates/UnrealTemplate.h"
+TAutoConsoleVariable<bool> CVarUseCurveCache(TEXT("CurveEditor.UseCurveCache"), true, TEXT("When true we cache curve values, when false we always regenerate."));
 
 SCurveEditorView::SCurveEditorView()
 	: bPinned(0)
@@ -28,7 +29,13 @@ SCurveEditorView::SCurveEditorView()
 	, bFixedOutputBounds(0)
 	, bAutoSize(1)
 	, bAllowEmpty(0)
-{}
+{
+	CurveCacheFlags = ECurveCacheFlags::All;
+	CachedValues.CachedActiveCurvesSerialNumber = 0xFFFFFFFF;
+	CachedValues.CachedSelectionSerialNumber = 0xFFFFFFFF;
+	CachedValues.CachedGeometrySize.X = -1.;
+	CachedValues.CachedGeometrySize.Y = -1;
+}
 
 FVector2D SCurveEditorView::ComputeDesiredSize(float LayoutScaleMultiplier) const
 {
@@ -153,152 +160,240 @@ void SCurveEditorView::GetCurveDrawParams(TArray<FCurveDrawParams>& OutDrawParam
 	double InputMin = 0, InputMax = 1;
 	GetInputBounds(InputMin, InputMax);
 
-	ECurveEditorTangentVisibility TangentVisibility = CurveEditor->GetSettings()->GetTangentVisibility();
-	OutDrawParams.Reserve(CurveInfoByID.Num());
+	OutDrawParams.Reset(CurveInfoByID.Num());
 
 	for (const TTuple<FCurveModelID, FCurveInfo>& Pair : CurveInfoByID)
 	{
+		FCurveDrawParams Params(Pair.Key);
+
 		FCurveModel* CurveModel = CurveEditor->FindCurve(Pair.Key);
 		if (!ensureAlways(CurveModel))
 		{
 			continue;
 		}
 
-		FCurveEditorScreenSpace CurveSpace = GetCurveSpace(Pair.Key);
+		GetCurveDrawParam(CurveEditor, Pair.Key, CurveModel,InputMin, InputMax, Params);
+		OutDrawParams.Add(MoveTemp(Params));
+	}
+}
 
-		const float DisplayRatio = (CurveSpace.PixelsPerOutput() / CurveSpace.PixelsPerInput());
+void SCurveEditorView::GetCurveDrawParam(TSharedPtr<FCurveEditor>& CurveEditor,const FCurveModelID& ModelID, FCurveModel* CurveModel, 
+	double InputMin, double InputMax, FCurveDrawParams& Params) const
+{
 
-		const FKeyHandleSet* SelectedKeys = CurveEditor->GetSelection().GetAll().Find(Pair.Key);
+	FCurveEditorScreenSpace CurveSpace = GetCurveSpace(ModelID);
+	const float DisplayRatio = (CurveSpace.PixelsPerOutput() / CurveSpace.PixelsPerInput());
 
-		// Create a new set of Curve Drawing Parameters to represent this particular Curve
-		FCurveDrawParams Params(Pair.Key);
-		Params.Color = CurveModel->GetColor();
-		Params.bKeyDrawEnabled = CurveModel->IsKeyDrawEnabled();
+	const FKeyHandleSet* SelectedKeys = CurveEditor->GetSelection().GetAll().Find(ModelID);
 
-		// Gather the display metrics to use for each key type. This allows a Curve Model to override
-		// whether or not the curve supports Keys, Arrive/Leave Tangents, etc. If the Curve Model doesn't
-		// support a particular capability we can skip drawing them.
-		CurveModel->GetKeyDrawInfo(ECurvePointType::ArriveTangent, FKeyHandle::Invalid(), Params.ArriveTangentDrawInfo);
-		CurveModel->GetKeyDrawInfo(ECurvePointType::LeaveTangent, FKeyHandle::Invalid(), Params.LeaveTangentDrawInfo);
+	// Create a new set of Curve Drawing Parameters to represent this particular Curve
+	Params.Color = CurveModel->GetColor();
+	Params.bKeyDrawEnabled = CurveModel->IsKeyDrawEnabled();
 
-		// Gather the interpolating points in input/output space
-		TArray<TTuple<double, double>> InterpolatingPoints;
+	// Gather the display metrics to use for each key type. This allows a Curve Model to override
+	// whether or not the curve supports Keys, Arrive/Leave Tangents, etc. If the Curve Model doesn't
+	// support a particular capability we can skip drawing them.
+	CurveModel->GetKeyDrawInfo(ECurvePointType::ArriveTangent, FKeyHandle::Invalid(), Params.ArriveTangentDrawInfo);
+	CurveModel->GetKeyDrawInfo(ECurvePointType::LeaveTangent, FKeyHandle::Invalid(), Params.LeaveTangentDrawInfo);
 
-		CurveModel->DrawCurve(*CurveEditor, CurveSpace, InterpolatingPoints);
-		Params.InterpolatingPoints.Reserve(InterpolatingPoints.Num());
+	// Gather the interpolating points in input/output space
+	TArray<TTuple<double, double>> InterpolatingPoints;
 
-		// An Input Offset allows for a fixed offset to all keys, such as displaying them in the middle of a frame instead of at the start.
-		double InputOffset = CurveModel->GetInputDisplayOffset();
+	CurveModel->DrawCurve(*CurveEditor, CurveSpace, InterpolatingPoints);
+	Params.InterpolatingPoints.Reset(InterpolatingPoints.Num());
 
-		// Convert the interpolating points to screen space
-		for (TTuple<double, double> Point : InterpolatingPoints)
-		{
-			Params.InterpolatingPoints.Add(
-				FVector2D(
-					CurveSpace.SecondsToScreen(Point.Get<0>() + InputOffset),
-					CurveSpace.ValueToScreen(Point.Get<1>())
-				)
-			);
-		}
+	// An Input Offset allows for a fixed offset to all keys, such as displaying them in the middle of a frame instead of at the start.
+	double InputOffset = CurveModel->GetInputDisplayOffset();
+
+	// Convert the interpolating points to screen space
+	for (TTuple<double, double> Point : InterpolatingPoints)
+	{
+		Params.InterpolatingPoints.Add(
+			FVector2D(
+				CurveSpace.SecondsToScreen(Point.Get<0>() + InputOffset),
+				CurveSpace.ValueToScreen(Point.Get<1>())
+			)
+		);
+	}
+
+	TArray<FKeyHandle> VisibleKeys;
+	CurveModel->GetKeys(*CurveEditor, InputMin, InputMax, TNumericLimits<double>::Lowest(), TNumericLimits<double>::Max(), VisibleKeys);
+
+	if (VisibleKeys.Num())
+	{
+		ECurveEditorTangentVisibility TangentVisibility = CurveEditor->GetSettings()->GetTangentVisibility();
 
-		TArray<FKeyHandle> VisibleKeys;
-		CurveModel->GetKeys(*CurveEditor, InputMin, InputMax, TNumericLimits<double>::Lowest(), TNumericLimits<double>::Max(), VisibleKeys);
+		TArray<FKeyPosition> AllKeyPositions;
+		TArray<FKeyAttributes> AllKeyAttributes;
 
-		if (VisibleKeys.Num())
+		AllKeyPositions.SetNum(VisibleKeys.Num());
+		AllKeyAttributes.SetNum(VisibleKeys.Num());
+		Params.Points.Reset(VisibleKeys.Num());
+
+		CurveModel->GetKeyPositions(VisibleKeys, AllKeyPositions);
+		CurveModel->GetKeyAttributes(VisibleKeys, AllKeyAttributes);
+		for (int32 Index = 0; Index < VisibleKeys.Num(); ++Index)
 		{
-			TArray<FKeyPosition> AllKeyPositions;
-			TArray<FKeyAttributes> AllKeyAttributes;
+			const FKeyHandle      KeyHandle = VisibleKeys[Index];
+			const FKeyPosition& KeyPosition = AllKeyPositions[Index];
+			const FKeyAttributes& Attributes = AllKeyAttributes[Index];
 
-			AllKeyPositions.SetNum(VisibleKeys.Num());
-			AllKeyAttributes.SetNum(VisibleKeys.Num());
+			bool bShowTangents = TangentVisibility == ECurveEditorTangentVisibility::AllTangents ||
+				(TangentVisibility == ECurveEditorTangentVisibility::SelectedKeys && SelectedKeys &&
+					(SelectedKeys->Contains(VisibleKeys[Index], ECurvePointType::Any)));
 
-			CurveModel->GetKeyPositions(VisibleKeys, AllKeyPositions);
-			CurveModel->GetKeyAttributes(VisibleKeys, AllKeyAttributes);
+			float TimeScreenPos = CurveSpace.SecondsToScreen(KeyPosition.InputValue + InputOffset);
+			float ValueScreenPos = CurveSpace.ValueToScreen(KeyPosition.OutputValue);
 
-			for (int32 Index = 0; Index < VisibleKeys.Num(); ++Index)
-			{
-				const FKeyHandle      KeyHandle = VisibleKeys[Index];
-				const FKeyPosition&   KeyPosition = AllKeyPositions[Index];
-				const FKeyAttributes& Attributes = AllKeyAttributes[Index];
+			// Add this key
+			FCurvePointInfo Key(KeyHandle);
+			Key.ScreenPosition = FVector2D(TimeScreenPos, ValueScreenPos);
+			Key.LayerBias = 2;
 
-				bool bShowTangents = TangentVisibility == ECurveEditorTangentVisibility::AllTangents || 
-					(TangentVisibility == ECurveEditorTangentVisibility::SelectedKeys && SelectedKeys && 
-					 (SelectedKeys->Contains(VisibleKeys[Index], ECurvePointType::Any) ) );
+			// Add draw info for the specific key
+			CurveModel->GetKeyDrawInfo(ECurvePointType::Key, KeyHandle, /*Out*/ Key.DrawInfo);
+			Params.Points.Add(Key);
 
-				float TimeScreenPos = CurveSpace.SecondsToScreen(KeyPosition.InputValue + InputOffset);
-				float ValueScreenPos = CurveSpace.ValueToScreen(KeyPosition.OutputValue);
+			if (bShowTangents && Attributes.HasArriveTangent())
+			{
+				float ArriveTangent = Attributes.GetArriveTangent();
 
-				// Add this key
-				FCurvePointInfo Key(KeyHandle);
-				Key.ScreenPosition = FVector2D(TimeScreenPos, ValueScreenPos);
-				Key.LayerBias = 2;
+				FCurvePointInfo ArriveTangentPoint(KeyHandle);
+				ArriveTangentPoint.Type = ECurvePointType::ArriveTangent;
 
-				// Add draw info for the specific key
-				CurveModel->GetKeyDrawInfo(ECurvePointType::Key, KeyHandle, /*Out*/ Key.DrawInfo);
-				Params.Points.Add(Key);
 
-				if (bShowTangents && Attributes.HasArriveTangent())
+				if (Attributes.HasTangentWeightMode() && Attributes.HasArriveTangentWeight() &&
+					(Attributes.GetTangentWeightMode() == RCTWM_WeightedBoth || Attributes.GetTangentWeightMode() == RCTWM_WeightedArrive))
+				{
+					FVector2D TangentOffset = CurveEditor::ComputeScreenSpaceTangentOffset(CurveSpace, ArriveTangent, -Attributes.GetArriveTangentWeight());
+					ArriveTangentPoint.ScreenPosition = Key.ScreenPosition + TangentOffset;
+				}
+				else
 				{
-					float ArriveTangent = Attributes.GetArriveTangent();
-
-					FCurvePointInfo ArriveTangentPoint(KeyHandle);
-					ArriveTangentPoint.Type = ECurvePointType::ArriveTangent;
-
-
-					if (Attributes.HasTangentWeightMode() && Attributes.HasArriveTangentWeight() &&
-						(Attributes.GetTangentWeightMode() == RCTWM_WeightedBoth || Attributes.GetTangentWeightMode() == RCTWM_WeightedArrive))
-					{
-						FVector2D TangentOffset = CurveEditor::ComputeScreenSpaceTangentOffset(CurveSpace, ArriveTangent, -Attributes.GetArriveTangentWeight());
-						ArriveTangentPoint.ScreenPosition = Key.ScreenPosition + TangentOffset;
-					}
-					else
-					{
-						float PixelLength = 60.0f;
-						ArriveTangentPoint.ScreenPosition = Key.ScreenPosition + CurveEditor::GetVectorFromSlopeAndLength(ArriveTangent * -DisplayRatio, -PixelLength);
-					}
-					ArriveTangentPoint.LineDelta = Key.ScreenPosition - ArriveTangentPoint.ScreenPosition;
-					ArriveTangentPoint.LayerBias = 1;
-
-					// Add draw info for the specific tangent
-					FKeyDrawInfo TangentDrawInfo;
-					CurveModel->GetKeyDrawInfo(ECurvePointType::ArriveTangent, KeyHandle, /*Out*/ ArriveTangentPoint.DrawInfo);
-
-					Params.Points.Add(ArriveTangentPoint);
+					float PixelLength = 60.0f;
+					ArriveTangentPoint.ScreenPosition = Key.ScreenPosition + CurveEditor::GetVectorFromSlopeAndLength(ArriveTangent * -DisplayRatio, -PixelLength);
 				}
+				ArriveTangentPoint.LineDelta = Key.ScreenPosition - ArriveTangentPoint.ScreenPosition;
+				ArriveTangentPoint.LayerBias = 1;
+
+				// Add draw info for the specific tangent
+				FKeyDrawInfo TangentDrawInfo;
+				CurveModel->GetKeyDrawInfo(ECurvePointType::ArriveTangent, KeyHandle, /*Out*/ ArriveTangentPoint.DrawInfo);
+
+				Params.Points.Add(ArriveTangentPoint);
+			}
 
-				if (bShowTangents && Attributes.HasLeaveTangent())
+			if (bShowTangents && Attributes.HasLeaveTangent())
+			{
+				float LeaveTangent = Attributes.GetLeaveTangent();
+
+				FCurvePointInfo LeaveTangentPoint(KeyHandle);
+				LeaveTangentPoint.Type = ECurvePointType::LeaveTangent;
+
+				if (Attributes.HasTangentWeightMode() && Attributes.HasLeaveTangentWeight() &&
+					(Attributes.GetTangentWeightMode() == RCTWM_WeightedBoth || Attributes.GetTangentWeightMode() == RCTWM_WeightedLeave))
 				{
-					float LeaveTangent = Attributes.GetLeaveTangent();
+					FVector2D TangentOffset = CurveEditor::ComputeScreenSpaceTangentOffset(CurveSpace, LeaveTangent, Attributes.GetLeaveTangentWeight());
 
-					FCurvePointInfo LeaveTangentPoint(KeyHandle);
-					LeaveTangentPoint.Type = ECurvePointType::LeaveTangent;
+					LeaveTangentPoint.ScreenPosition = Key.ScreenPosition + TangentOffset;
+				}
+				else
+				{
+					float PixelLength = 60.0f;
+					LeaveTangentPoint.ScreenPosition = Key.ScreenPosition + CurveEditor::GetVectorFromSlopeAndLength(LeaveTangent * -DisplayRatio, PixelLength);
 
-					if (Attributes.HasTangentWeightMode() && Attributes.HasLeaveTangentWeight() &&
-						(Attributes.GetTangentWeightMode() == RCTWM_WeightedBoth || Attributes.GetTangentWeightMode() == RCTWM_WeightedLeave))
-					{
-						FVector2D TangentOffset = CurveEditor::ComputeScreenSpaceTangentOffset(CurveSpace, LeaveTangent, Attributes.GetLeaveTangentWeight());
+				}
 
-						LeaveTangentPoint.ScreenPosition = Key.ScreenPosition + TangentOffset;
-					}
-					else
-					{
-						float PixelLength = 60.0f;
-						LeaveTangentPoint.ScreenPosition = Key.ScreenPosition + CurveEditor::GetVectorFromSlopeAndLength(LeaveTangent * -DisplayRatio, PixelLength);
+				LeaveTangentPoint.LineDelta = Key.ScreenPosition - LeaveTangentPoint.ScreenPosition;
+				LeaveTangentPoint.LayerBias = 1;
 
-					}
+				// Add draw info for the specific tangent
+				FKeyDrawInfo TangentDrawInfo;
+				CurveModel->GetKeyDrawInfo(ECurvePointType::LeaveTangent, KeyHandle, /*Out*/ LeaveTangentPoint.DrawInfo);
 
-					LeaveTangentPoint.LineDelta = Key.ScreenPosition - LeaveTangentPoint.ScreenPosition;
-					LeaveTangentPoint.LayerBias = 1;
+				Params.Points.Add(LeaveTangentPoint);
+			}
+		}
+	}
+}
+
+void SCurveEditorView::Tick(const FGeometry& AllottedGeometry, const double InCurrentTime, const float InDeltaTime)
+{
+	TSharedPtr<FCurveEditor> CurveEditor = WeakCurveEditor.Pin();
+	check(CurveEditor.IsValid());
+	const bool bUseCurveCache = CVarUseCurveCache.GetValueOnGameThread();
 
-					// Add draw info for the specific tangent
-					FKeyDrawInfo TangentDrawInfo;
-					CurveModel->GetKeyDrawInfo(ECurvePointType::LeaveTangent, KeyHandle, /*Out*/ LeaveTangentPoint.DrawInfo);
+	if (bUseCurveCache)
+	{
+		//if number of curves have changed, just redo all
+		if (CurveEditor->GetActiveCurvesSerialNumber() != CachedValues.CachedActiveCurvesSerialNumber)
+		{
+			CachedValues.CachedActiveCurvesSerialNumber = CurveEditor->GetActiveCurvesSerialNumber();
+			CurveCacheFlags = ECurveCacheFlags::All;
+		}
+		if (CachedValues.CachedTangentVisibility != CurveEditor->GetSettings()->GetTangentVisibility())
+		{
+			CachedValues.CachedTangentVisibility = CurveEditor->GetSettings()->GetTangentVisibility();
+			CurveCacheFlags = ECurveCacheFlags::All;
+		}
+		if (CachedValues.CachedSelectionSerialNumber != CurveEditor->GetSelection().GetSerialNumber())
+		{
+			CachedValues.CachedSelectionSerialNumber = CurveEditor->GetSelection().GetSerialNumber();
+			CurveCacheFlags = ECurveCacheFlags::All;
+		}
 
-					Params.Points.Add(LeaveTangentPoint);
+		//Only get view values if we need to since we will reset them every time we get all
+		if (CurveCacheFlags != ECurveCacheFlags::All)
+		{
+			if (OutputMin != CachedValues.CachedOutputMin || OutputMax != CachedValues.CachedOutputMax)
+			{
+				CurveCacheFlags = ECurveCacheFlags::All;
+			}
+			else if (CachedValues.CachedGeometrySize != GetCachedGeometry().GetLocalSize())
+			{
+				CurveCacheFlags = ECurveCacheFlags::All;
+			}
+			else
+			{
+				double InputMin = 0, InputMax = 1;
+				GetInputBounds(InputMin, InputMax);
+				if (InputMin != CachedValues.CachedInputMin || InputMax != CachedValues.CachedInputMax)
+				{
+					CurveCacheFlags = ECurveCacheFlags::All;
+				}
+			}
+		}
+		if (CurveCacheFlags == ECurveCacheFlags::All)
+		{
+			CachedValues.CachedOutputMin = OutputMin;
+			CachedValues.CachedOutputMax = OutputMax;
+			GetInputBounds(CachedValues.CachedInputMin, CachedValues.CachedInputMax);
+			CachedValues.CachedGeometrySize = GetCachedGeometry().GetLocalSize();
+
+			CachedDrawParams.Reset();
+			GetCurveDrawParams(CachedDrawParams);
+			CurveCacheFlags = ECurveCacheFlags::CheckCurves;
+		}
+		else if (CurveCacheFlags == ECurveCacheFlags::CheckCurves)
+		{
+			for (FCurveDrawParams& Params : CachedDrawParams)
+			{
+				FCurveModel* CurveModel = CurveEditor->FindCurve(Params.GetID());
+				if (CurveModel->HasChangedAndResetTest())
+				{
+					GetCurveDrawParam(CurveEditor, Params.GetID(), CurveModel, CachedValues.CachedInputMin, CachedValues.CachedInputMax, Params);
 				}
 			}
 		}
+	}
+	else
+	{
+		CachedValues.CachedOutputMin = OutputMin;
+		CachedValues.CachedOutputMax = OutputMax;
+		GetInputBounds(CachedValues.CachedInputMin, CachedValues.CachedInputMax);
+		CachedValues.CachedGeometrySize = GetCachedGeometry().GetLocalSize();
 
-		OutDrawParams.Add(MoveTemp(Params));
+		CachedDrawParams.Reset();
+		GetCurveDrawParams(CachedDrawParams);
 	}
-}
+}
\ No newline at end of file
diff --git a/Engine/Source/Editor/CurveEditor/Private/SCurveViewerPanel.cpp b/Engine/Source/Editor/CurveEditor/Private/SCurveViewerPanel.cpp
--- a/Engine/Source/Editor/CurveEditor/Private/SCurveViewerPanel.cpp
+++ b/Engine/Source/Editor/CurveEditor/Private/SCurveViewerPanel.cpp
@@ -38,6 +38,7 @@ void SCurveViewerPanel::Construct(const FArguments& InArgs, TSharedRef<FCurveEdi
 	CurveThickness = InArgs._CurveThickness;
 
 	InCurveEditor->SetView(SharedThis(this));
+	CachedValues.CachedTangentVisibility = InCurveEditor->GetSettings()->GetTangentVisibility();
 
 	for (const TPair<FCurveModelID, TUniquePtr<FCurveModel>>& CurvePair : InCurveEditor->GetCurves())
 	{
@@ -47,12 +48,6 @@ void SCurveViewerPanel::Construct(const FArguments& InArgs, TSharedRef<FCurveEdi
 	SetClipping(EWidgetClipping::ClipToBounds);
 }
 
-void SCurveViewerPanel::Tick(const FGeometry& AllottedGeometry, const double InCurrentTime, const float InDeltaTime)
-{
-	CachedDrawParams.Reset();
-	GetCurveDrawParams(CachedDrawParams);
-}
-
 int32 SCurveViewerPanel::OnPaint(const FPaintArgs& Args, const FGeometry& AllottedGeometry, const FSlateRect& MyCullingRect, FSlateWindowElementList& OutDrawElements, int32 LayerId, const FWidgetStyle& InWidgetStyle, bool bParentEnabled) const
 {
 	const ESlateDrawEffect DrawEffects = ShouldBeEnabled(bParentEnabled) ? ESlateDrawEffect::None : ESlateDrawEffect::DisabledEffect;
diff --git a/Engine/Source/Editor/CurveEditor/Private/Views/SInteractiveCurveEditorView.cpp b/Engine/Source/Editor/CurveEditor/Private/Views/SInteractiveCurveEditorView.cpp
--- a/Engine/Source/Editor/CurveEditor/Private/Views/SInteractiveCurveEditorView.cpp
+++ b/Engine/Source/Editor/CurveEditor/Private/Views/SInteractiveCurveEditorView.cpp
@@ -70,6 +70,9 @@
 class FPaintArgs;
 class FWidgetStyle;
 
+TAutoConsoleVariable<bool> CVarDrawCurveLines(TEXT("CurveEditor.DrawCurveLines"), true, TEXT("When true we draw curve lines, when false we do not."));
+TAutoConsoleVariable<bool> CVarDrawCurveKeys(TEXT("CurveEditor.DrawCurveKeys"), true, TEXT("When true we draw curve keys, when false we do not."));
+
 namespace CurveViewConstants
 {
 	/** The number of pixels to offset Labels from the Left/Right size. */
@@ -103,12 +106,14 @@ public:
 
 void SInteractiveCurveEditorView::Construct(const FArguments& InArgs, TWeakPtr<FCurveEditor> InCurveEditor)
 {
+
 	FixedHeight = InArgs._FixedHeight;
 	BackgroundTint = InArgs._BackgroundTint;
 	MaximumCapacity = InArgs._MaximumCapacity;
 	bAutoSize = InArgs._AutoSize;
 
 	WeakCurveEditor = InCurveEditor;
+	CachedValues.CachedTangentVisibility = InCurveEditor.Pin()->GetSettings()->GetTangentVisibility();
 
 	InCurveEditor.Pin()->OnActiveToolChangedDelegate.AddSP(this, &SInteractiveCurveEditorView::OnCurveEditorToolChanged);
 
@@ -436,6 +441,9 @@ void SInteractiveCurveEditorView::DrawCurves(TSharedRef<FCurveEditor> CurveEdito
 	const float UnHoveredCurveThickness = 2.f;
 	const bool  bAntiAliasCurves = true;
 
+	const bool bDrawLines = CVarDrawCurveLines.GetValueOnGameThread();
+	const bool bDrawKeys = CVarDrawCurveKeys.GetValueOnGameThread();
+
 	TOptional<FCurveModelID> HoveredCurve = GetHoveredCurve();
 	for (const FCurveDrawParams& Params : CachedDrawParams)
 	{
@@ -443,22 +451,26 @@ void SInteractiveCurveEditorView::DrawCurves(TSharedRef<FCurveEditor> CurveEdito
 		const float Thickness = bIsCurveHovered ? HoveredCurveThickness : UnHoveredCurveThickness;
 		const int32 CurveLayerId = bIsCurveHovered ? BaseLayerId + CurveViewConstants::ELayerOffset::Curves : BaseLayerId + CurveViewConstants::ELayerOffset::HoveredCurves;
 
-		FSlateDrawElement::MakeLines(
-			OutDrawElements,
-			CurveLayerId,
-			PaintGeometry,
-			Params.InterpolatingPoints,
-			DrawEffects,
-			Params.Color,
-			bAntiAliasCurves,
-			Thickness
-		);
+		if (bDrawLines)
+		{
 
+			FSlateDrawElement::MakeLines(
+				OutDrawElements,
+				CurveLayerId,
+				PaintGeometry,
+				Params.InterpolatingPoints,
+				DrawEffects,
+				Params.Color,
+				bAntiAliasCurves,
+				Thickness
+			);
+		}
+		
 		TArray<FVector2D> LinePoints;
 		LinePoints.SetNum(2);
 
 		// Draw tangents
-		if (Params.bKeyDrawEnabled)
+		if (bDrawKeys && Params.bKeyDrawEnabled)
 		{
 			for (int32 PointIndex = 0; PointIndex < Params.Points.Num(); PointIndex++)
 			{
@@ -551,13 +563,6 @@ void SInteractiveCurveEditorView::DrawBufferedCurves(TSharedRef<FCurveEditor> Cu
 	}
 }
 
-void SInteractiveCurveEditorView::Tick(const FGeometry& AllottedGeometry, const double InCurrentTime, const float InDeltaTime)
-{
-	// Cache our Curve Drawing Params. These are used in multiple places so we cache them once each frame.
-	CachedDrawParams.Reset();
-	GetCurveDrawParams(CachedDrawParams);
-}
-
 bool SInteractiveCurveEditorView::GetPointsWithinWidgetRange(const FSlateRect& WidgetRectangle, TArray<FCurvePointHandle>* OutPoints) const
 {
 	TSharedPtr<FCurveEditor> CurveEditor = WeakCurveEditor.Pin();
diff --git a/Engine/Source/Editor/CurveEditor/Public/CurveModel.h b/Engine/Source/Editor/CurveEditor/Public/CurveModel.h
--- a/Engine/Source/Editor/CurveEditor/Public/CurveModel.h
+++ b/Engine/Source/Editor/CurveEditor/Public/CurveModel.h
@@ -247,6 +247,12 @@ public:
 		return nullptr;
 	}
 
+	/** Get if has changed and then reset it, this can be used for caching*/
+	virtual bool HasChangedAndResetTest()
+	{
+		return true;
+	}
+
 	/**
 	* Get the Object and the name to be used to store the curve model color (see UCurveEditorSettings). By default
 	* this is the owning object and the intent name, but it can be overriden, for example for Sequencer it may be the bound object
@@ -357,9 +363,13 @@ public:
 
 	/**
 	 */
-	FORCEINLINE void SetColor(const FLinearColor& InColor)
+	FORCEINLINE void SetColor(const FLinearColor& InColor, bool bInModify = true)
 	{
 		Color = InColor;
+		if (bInModify)
+		{
+			Modify(); //will make sure the cache get's recreated
+		}
 	}
 
 	/**
diff --git a/Engine/Source/Editor/CurveEditor/Public/SCurveEditorView.h b/Engine/Source/Editor/CurveEditor/Public/SCurveEditorView.h
--- a/Engine/Source/Editor/CurveEditor/Public/SCurveEditorView.h
+++ b/Engine/Source/Editor/CurveEditor/Public/SCurveEditorView.h
@@ -17,12 +17,14 @@
 #include "Templates/SharedPointer.h"
 #include "Widgets/DeclarativeSyntaxSupport.h"
 #include "Widgets/SCompoundWidget.h"
+#include "CurveEditorSettings.h"
 
 class FCurveEditor;
 class FSlateRect;
 class FText;
 class SCurveEditorPanel;
 struct FCurveDrawParams;
+class FCurveModel;
 
 /**
  * This is the base widget type for all views that exist on a curve editor panel. A view may contain 0 or more curves (stored in CurveInfoByID).
@@ -207,7 +209,12 @@ protected:
 	/** Gets info about the curves being drawn. Converts actual curves into an abstract series of lines/points/handles/etc. */
 	void GetCurveDrawParams(TArray<FCurveDrawParams>& OutDrawParams) const;
 
+	/** Get it for just one curve*/
+	void GetCurveDrawParam(TSharedPtr<FCurveEditor>& CurveEditor, const FCurveModelID& ModelID, FCurveModel* CurveModel,
+		double InputMin, double InputMax, FCurveDrawParams& OutDrawParam) const;
+
 	// ~SWidget interface
+	virtual void Tick(const FGeometry& AllottedGeometry, const double InCurrentTime, const float InDeltaTime) override;
 	virtual FVector2D ComputeDesiredSize(float LayoutScaleMultiplier) const override;
 
 private:
@@ -270,4 +277,38 @@ protected:
 	 * @note: Should only be added to or removed from in AddCurve/RemoveCurve. Derived types must only change the FCurveInfo contained within this map.
 	 */
 	TSortedMap<FCurveModelID, FCurveInfo, TInlineAllocator<1>> CurveInfoByID;
+
+	/** Flag enum signifying how the curve cache has changed since it was last generated
+	Note for a data change it may only effect certain data(curves) not every drawn curve*/
+	enum class ECurveCacheFlags : uint8
+	{
+		CheckCurves = 0,       // The cache may be valid need to check each curve to see if they are still valid
+		All = 1 << 0,		   // Get all
+	};
+	/** Curve cache flags that change based upon data or view getting modified*/
+	ECurveCacheFlags CurveCacheFlags;
+
+	/** Curve draw parameters that are re-generated on tick if the cache has changed. We generate them once and then they're used in multiple places per frame. */
+	TArray<FCurveDrawParams> CachedDrawParams;
+
+	/** Set of Cached values we need to check each tick to see if we need to redo cache*/
+	struct FCachedValuesToCheck
+	{
+		/** Serial number cached from FCurveEditor::GetActiveCurvesSerialNumber() on tick */
+		uint32 CachedActiveCurvesSerialNumber;
+
+		/** Serial number bached from CurveEditorSelecstion::GetSerialNumber */
+		uint32 CachedSelectionSerialNumber;
+
+		/** Cached Tangent Visibility*/
+		ECurveEditorTangentVisibility CachedTangentVisibility;
+
+		/** Cached input and output min max values to see if we need to recalc curves, though we need to poll it's safer*/
+		double CachedInputMin, CachedInputMax, CachedOutputMin, CachedOutputMax;
+
+		/** Cached Geometry Size*/
+		FVector2D CachedGeometrySize;
+	};
+
+	FCachedValuesToCheck CachedValues;
 };
\ No newline at end of file
diff --git a/Engine/Source/Editor/CurveEditor/Public/SCurveViewerPanel.h b/Engine/Source/Editor/CurveEditor/Public/SCurveViewerPanel.h
--- a/Engine/Source/Editor/CurveEditor/Public/SCurveViewerPanel.h
+++ b/Engine/Source/Editor/CurveEditor/Public/SCurveViewerPanel.h
@@ -48,7 +48,6 @@ class CURVEEDITOR_API SCurveViewerPanel : public SCurveEditorView
 private:
 
 	virtual int32 OnPaint(const FPaintArgs& Args, const FGeometry& AllottedGeometry, const FSlateRect& MyCullingRect, FSlateWindowElementList& OutDrawElements, int32 LayerId, const FWidgetStyle& InWidgetStyle, bool bParentEnabled) const override;
-	virtual void Tick(const FGeometry& AllottedGeometry, const double InCurrentTime, const float InDeltaTime) override;
 
 public:
 
diff --git a/Engine/Source/Editor/CurveEditor/Public/Views/SInteractiveCurveEditorView.h b/Engine/Source/Editor/CurveEditor/Public/Views/SInteractiveCurveEditorView.h
--- a/Engine/Source/Editor/CurveEditor/Public/Views/SInteractiveCurveEditorView.h
+++ b/Engine/Source/Editor/CurveEditor/Public/Views/SInteractiveCurveEditorView.h
@@ -24,6 +24,7 @@
 #include "Templates/Tuple.h"
 #include "Widgets/DeclarativeSyntaxSupport.h"
 #include "Widgets/SCompoundWidget.h"
+#include "CurveEditorSettings.h"
 
 class FCurveEditor;
 class FCurveModel;
@@ -117,7 +118,6 @@ protected:
 
 	// SWidget Interface
 	virtual int32 OnPaint(const FPaintArgs& Args, const FGeometry& AllottedGeometry, const FSlateRect& MyCullingRect, FSlateWindowElementList& OutDrawElements, int32 BaseLayerId, const FWidgetStyle& InWidgetStyle, bool bParentEnabled) const override;
-	virtual void Tick(const FGeometry& AllottedGeometry, const double InCurrentTime, const float InDeltaTime) override;
 	virtual void OnMouseEnter(const FGeometry& MyGeometry, const FPointerEvent& MouseEvent) override;
 	virtual void OnMouseLeave(const FPointerEvent& MouseEvent) override;
 	virtual FReply OnMouseMove(const FGeometry& MyGeometry, const FPointerEvent& MouseEvent) override;
@@ -191,9 +191,6 @@ protected:
 
 private:
 
-	/** Curve draw parameters that are re-generated on tick. We generate them once and then they're used in multiple places per frame. */
-	TArray<FCurveDrawParams> CachedDrawParams;
-
 	/** (Optional) the current drag operation */
 	TOptional<FCurveEditorDelayedDrag> DragOperation;
 
diff --git a/Engine/Source/Editor/MovieSceneTools/Private/Channels/ChannelCurveModel.cpp b/Engine/Source/Editor/MovieSceneTools/Private/Channels/ChannelCurveModel.cpp
--- a/Engine/Source/Editor/MovieSceneTools/Private/Channels/ChannelCurveModel.cpp
+++ b/Engine/Source/Editor/MovieSceneTools/Private/Channels/ChannelCurveModel.cpp
@@ -38,6 +38,7 @@ FChannelCurveModel<ChannelType, ChannelValue, KeyType>::FChannelCurveModel(TMovi
 	ChannelHandle = InChannel;
 	WeakSection = OwningSection;
 	WeakSequencer = InWeakSequencer;
+	LastSignature = OwningSection->GetSignature();
 
 	if (FMovieSceneChannelProxy* ChannelProxy = InChannel.GetChannelProxy())
 	{
@@ -69,6 +70,7 @@ void FChannelCurveModel<ChannelType, ChannelValue, KeyType>::Modify()
 	{
 		Section->Modify();
 	}
+	LastSignature.Invalidate();
 }
 
 
@@ -439,6 +441,8 @@ void FChannelCurveModel<ChannelType, ChannelValue, KeyType>::FixupCurve()
 		OnDestroyHandle = NewChannelProxy->OnDestroy.AddRaw(this, &FChannelCurveModel<ChannelType, ChannelValue, KeyType>::FixupCurve);
 	}
 }
+
+
 template <class ChannelType, class ChannelValue, class KeyType>
 void FChannelCurveModel<ChannelType, ChannelValue, KeyType>::GetCurveColorObjectAndName(UObject** OutObject, FString& OutName) const
 {
diff --git a/Engine/Source/Editor/MovieSceneTools/Private/Channels/ConstraintChannelCurveModel.cpp b/Engine/Source/Editor/MovieSceneTools/Private/Channels/ConstraintChannelCurveModel.cpp
--- a/Engine/Source/Editor/MovieSceneTools/Private/Channels/ConstraintChannelCurveModel.cpp
+++ b/Engine/Source/Editor/MovieSceneTools/Private/Channels/ConstraintChannelCurveModel.cpp
@@ -50,12 +50,32 @@ FConstraintChannelCurveModel::FConstraintChannelCurveModel(TMovieSceneChannelHan
 	}
 
 	WeakSection = OwningSection;
+	LastSignature = OwningSection->GetSignature();
 	WeakSequencer = InWeakSequencer;
 	SupportedViews = ViewID;
 	Constraint = nullptr;
 	Color = STrackAreaView::BlendDefaultTrackColor(OwningSection->GetTypedOuter<UMovieSceneTrack>()->GetColorTint());
 }
 
+UObject* FConstraintChannelCurveModel::GetOwningObject() const
+{
+	return WeakSection.Get();
+}
+
+bool FConstraintChannelCurveModel::HasChangedAndResetTest()
+{
+	if (UMovieSceneSection* Section = WeakSection.Get())
+	{
+		if (Section->GetSignature() != LastSignature)
+		{
+			LastSignature = Section->GetSignature();
+			return true;
+		}
+		return false;
+	}
+	return true;
+}
+
 FConstraintChannelCurveModel::~FConstraintChannelCurveModel()
 {
 	if (FMovieSceneChannelProxy* ChannelProxy = ChannelHandle.GetChannelProxy())
diff --git a/Engine/Source/Editor/MovieSceneTools/Public/Channels/ChannelCurveModel.h b/Engine/Source/Editor/MovieSceneTools/Public/Channels/ChannelCurveModel.h
--- a/Engine/Source/Editor/MovieSceneTools/Public/Channels/ChannelCurveModel.h
+++ b/Engine/Source/Editor/MovieSceneTools/Public/Channels/ChannelCurveModel.h
@@ -21,6 +21,7 @@
 #include "UObject/UnrealType.h"
 #include "UObject/WeakObjectPtr.h"
 #include "UObject/WeakObjectPtrTemplates.h"
+#include "Misc/Guid.h"
 
 class FCurveEditor;
 class FString;
@@ -66,6 +67,20 @@ public:
 	{
 		return WeakSection.Get();
 	}
+	virtual bool HasChangedAndResetTest() override
+	{
+		if (UMovieSceneSection* Section = WeakSection.Get())
+		{
+			if (Section->GetSignature() != LastSignature)
+			{
+				LastSignature = Section->GetSignature();
+				return true;
+			}
+			return false;
+		}
+		return true;
+	}
+
 	virtual void GetCurveColorObjectAndName(UObject** OutObject, FString& OutName) const override;
 
 protected:
@@ -87,4 +102,5 @@ private:
 	TWeakObjectPtr<UMovieSceneSection> WeakSection;
 	TWeakPtr<ISequencer> WeakSequencer;
 	FDelegateHandle OnDestroyHandle;
+	FGuid LastSignature;
 };
\ No newline at end of file
diff --git a/Engine/Source/Editor/MovieSceneTools/Public/Channels/ConstraintChannelCurveModel.h b/Engine/Source/Editor/MovieSceneTools/Public/Channels/ConstraintChannelCurveModel.h
--- a/Engine/Source/Editor/MovieSceneTools/Public/Channels/ConstraintChannelCurveModel.h
+++ b/Engine/Source/Editor/MovieSceneTools/Public/Channels/ConstraintChannelCurveModel.h
@@ -65,7 +65,10 @@ public:
 	virtual void RemoveKeys(TArrayView<const FKeyHandle> InKeys) override;
 
 	virtual void CreateKeyProxies(TArrayView<const FKeyHandle> InKeyHandles, TArrayView<UObject*> OutObjects) override;
-
+	
+	virtual UObject* GetOwningObject() const override;
+	virtual bool HasChangedAndResetTest() override;
+	
 	//FKeyBarCurveModel
 	virtual void BuildContextMenu(const FCurveEditor& CurveEditor, FMenuBuilder& MenuBuilder, TOptional<FCurvePointHandle> ClickedPoint);
 	virtual TArray<FKeyBarCurveModel::FBarRange> FindRanges();
@@ -81,6 +84,7 @@ private:
 private:
 	mutable TMovieSceneChannelHandle<FMovieSceneConstraintChannel> ChannelHandle;
 	TWeakObjectPtr<UMovieSceneSection> WeakSection;
+	FGuid LastSignature;
 	TWeakPtr<ISequencer> WeakSequencer;
 	TWeakObjectPtr<UTickableConstraint> Constraint;
 
diff --git a/Engine/Source/Editor/Sequencer/Private/IKeyArea.cpp b/Engine/Source/Editor/Sequencer/Private/IKeyArea.cpp
--- a/Engine/Source/Editor/Sequencer/Private/IKeyArea.cpp
+++ b/Engine/Source/Editor/Sequencer/Private/IKeyArea.cpp
@@ -330,7 +330,7 @@ TUniquePtr<FCurveModel> IKeyArea::CreateCurveEditorModel(TSharedRef<ISequencer>
 			CurveModel->SetLongIntentionName(LongIntentName);
 			if (Color.IsSet())
 			{
-				CurveModel->SetColor(Color.GetValue());
+				CurveModel->SetColor(Color.GetValue(),false);
 			}
 
 			//Use editor preference color if it's been set, this function is const so we can't set just Color optional
