Sequencer: Curve Editor: Improve curve editor playback speed by using an UMG SRetainerWidget that caches to a render target. We then tie the new caching of our lines/keys so that if our cache updates we just then need to manually RequestRender on the RetainerWidget. This gives us nearly no performance decrease now when rendering curves. In AssetTools it uses the CurveEditorView module causing a cycle, but it's not actually needed since it loads and uses the module directly and uses privateincludedepency instead.  Currently we need to include UMG but this will probably change as Daren get's time to move the RetainerWidget code to a different module.

diff --git a/Engine/Source/Developer/AssetTools/AssetTools.Build.cs b/Engine/Source/Developer/AssetTools/AssetTools.Build.cs
--- a/Engine/Source/Developer/AssetTools/AssetTools.Build.cs
+++ b/Engine/Source/Developer/AssetTools/AssetTools.Build.cs
@@ -25,7 +25,6 @@ public class AssetTools : ModuleRules
 		PrivateDependencyModuleNames.AddRange(
 			new string[] {
 				"AssetDefinition",
-                "CurveAssetEditor",
 				"Engine",
                 "InputCore",
 				"ApplicationCore",
diff --git a/Engine/Source/Editor/CurveEditor/CurveEditor.Build.cs b/Engine/Source/Editor/CurveEditor/CurveEditor.Build.cs
--- a/Engine/Source/Editor/CurveEditor/CurveEditor.Build.cs
+++ b/Engine/Source/Editor/CurveEditor/CurveEditor.Build.cs
@@ -18,6 +18,7 @@ public class CurveEditor : ModuleRules
 				"Slate",
 				"SlateCore",
 				"TimeManagement",
+				"UMG",
 				"UnrealEd",
 				"SequencerWidgets",
 				"ToolMenus",
diff --git a/Engine/Source/Editor/CurveEditor/Private/SCurveEditorView.cpp b/Engine/Source/Editor/CurveEditor/Private/SCurveEditorView.cpp
--- a/Engine/Source/Editor/CurveEditor/Private/SCurveEditorView.cpp
+++ b/Engine/Source/Editor/CurveEditor/Private/SCurveEditorView.cpp
@@ -19,8 +19,10 @@
 #include "Math/Color.h"
 #include "Misc/AssertionMacros.h"
 #include "SCurveEditorPanel.h"
+#include "Slate/SRetainerWidget.h"
 #include "Templates/Tuple.h"
 #include "Templates/UnrealTemplate.h"
+
 TAutoConsoleVariable<bool> CVarUseCurveCache(TEXT("CurveEditor.UseCurveCache"), true, TEXT("When true we cache curve values, when false we always regenerate."));
 
 SCurveEditorView::SCurveEditorView()
@@ -317,7 +319,7 @@ void SCurveEditorView::GetCurveDrawParam(TSharedPtr<FCurveEditor>& CurveEditor,c
 	}
 }
 
-void SCurveEditorView::Tick(const FGeometry& AllottedGeometry, const double InCurrentTime, const float InDeltaTime)
+void SCurveEditorView::CheckCacheAndInvalidateIfNeeded()
 {
 	TSharedPtr<FCurveEditor> CurveEditor = WeakCurveEditor.Pin();
 	check(CurveEditor.IsValid());
@@ -373,17 +375,27 @@ void SCurveEditorView::Tick(const FGeometry& AllottedGeometry, const double InCu
 			CachedDrawParams.Reset();
 			GetCurveDrawParams(CachedDrawParams);
 			CurveCacheFlags = ECurveCacheFlags::CheckCurves;
+			if (RetainerWidget)
+			{
+				RetainerWidget->RequestRender();
+			}
 		}
 		else if (CurveCacheFlags == ECurveCacheFlags::CheckCurves)
 		{
+			bool bSomethingChanged = false;
 			for (FCurveDrawParams& Params : CachedDrawParams)
 			{
 				FCurveModel* CurveModel = CurveEditor->FindCurve(Params.GetID());
 				if (CurveModel->HasChangedAndResetTest())
 				{
+					bSomethingChanged = true;
 					GetCurveDrawParam(CurveEditor, Params.GetID(), CurveModel, CachedValues.CachedInputMin, CachedValues.CachedInputMax, Params);
 				}
 			}
+			if (bSomethingChanged && RetainerWidget)
+			{
+				RetainerWidget->RequestRender();
+			}
 		}
 	}
 	else
@@ -395,5 +407,10 @@ void SCurveEditorView::Tick(const FGeometry& AllottedGeometry, const double InCu
 
 		CachedDrawParams.Reset();
 		GetCurveDrawParams(CachedDrawParams);
+
+		if (RetainerWidget)
+		{
+			RetainerWidget->RequestRender();
+		}
 	}
 }
\ No newline at end of file
diff --git a/Engine/Source/Editor/CurveEditor/Private/SCurveEditorViewContainer.cpp b/Engine/Source/Editor/CurveEditor/Private/SCurveEditorViewContainer.cpp
--- a/Engine/Source/Editor/CurveEditor/Private/SCurveEditorViewContainer.cpp
+++ b/Engine/Source/Editor/CurveEditor/Private/SCurveEditorViewContainer.cpp
@@ -25,6 +25,7 @@
 #include "Rendering/RenderingCommon.h"
 #include "SCurveEditorPanel.h"
 #include "SCurveEditorView.h"
+#include "Slate/SRetainerWidget.h"
 #include "SlotBase.h"
 #include "Styling/AppStyle.h"
 #include "Styling/SlateBrush.h"
@@ -101,6 +102,15 @@ void SCurveEditorViewContainer::Tick(const FGeometry& AllottedGeometry, const do
 {
 	ExpandInputBounds(AllottedGeometry.GetLocalSize().X);
 
+	//we manually check the child views since we can't rely on ::Tick being called since the RetainerWidget will stop them if not actually rendering
+	for (TSharedPtr<SCurveEditorView>& View : Views)
+	{
+		if (View)
+		{
+			View->CheckCacheAndInvalidateIfNeeded();
+		}
+	}
+
 	if (CurveEditor->GetCurrentTool())
 	{
 		CurveEditor->GetCurrentTool()->Tick(AllottedGeometry, InCurrentTime, InDeltaTime);
@@ -355,13 +365,19 @@ void SCurveEditorViewContainer::AddView(TSharedRef<SCurveEditorView> ViewToAdd)
 	Views.Add(ViewToAdd);
 	SVerticalBox::FSlot* SlotPointer = nullptr;
 	AddSlot()
-	.Expose(SlotPointer)
-	[
-		SNew(SBox)
-		.Padding(MakeAttributeSP(this, &SCurveEditorViewContainer::GetSlotPadding, InsertIndex))
-		.Clipping(EWidgetClipping::ClipToBounds)
+		.Expose(SlotPointer)
+		[
+		SAssignNew(RetainerWidget, SRetainerWidget)
+		.RenderOnPhase(false)
+		.RenderOnInvalidation(false)
+		.RenderWithLocalTransform(false)
 		[
-			ViewToAdd
+			SNew(SBox)
+			.Padding(MakeAttributeSP(this, &SCurveEditorViewContainer::GetSlotPadding, InsertIndex))
+			.Clipping(EWidgetClipping::ClipToBounds)
+			[
+				ViewToAdd
+			]
 		]
 	];
 
@@ -369,6 +385,8 @@ void SCurveEditorViewContainer::AddView(TSharedRef<SCurveEditorView> ViewToAdd)
 	{
 		SlotPointer->SetAutoHeight();
 	}
+
+	ViewToAdd->SetRetainerWidget(RetainerWidget);
 }
 
 void SCurveEditorViewContainer::Clear()
diff --git a/Engine/Source/Editor/CurveEditor/Private/SCurveEditorViewContainer.h b/Engine/Source/Editor/CurveEditor/Private/SCurveEditorViewContainer.h
--- a/Engine/Source/Editor/CurveEditor/Private/SCurveEditorViewContainer.h
+++ b/Engine/Source/Editor/CurveEditor/Private/SCurveEditorViewContainer.h
@@ -22,6 +22,7 @@ class FSlateWindowElementList;
 class FWidgetStyle;
 class ITimeSliderController;
 class SCurveEditorView;
+class SRetainerWidget;
 struct FCurveEditorToolID;
 struct FFocusEvent;
 struct FGeometry;
@@ -95,9 +96,12 @@ private:
 	/** (Optional) the current drag operation */
 	TOptional<FCurveEditorDelayedDrag> DragOperation;
 
-	/** Array of views that may need their height updating on tick. */
+	/** Array of views whose cache may need to be updated and may need their height updating on tick. */
 	TArray<TSharedPtr<SCurveEditorView>> Views;
 
+	/** Possible pointer to a retainer widget that we may need to force update*/
+	TSharedPtr<SRetainerWidget> RetainerWidget;
+
 	/** 
 	 * Whether or not this widget caught an OnMouseDown notification 
 	 * Used to check if the selection should be cleared
diff --git a/Engine/Source/Editor/CurveEditor/Public/SCurveEditorView.h b/Engine/Source/Editor/CurveEditor/Public/SCurveEditorView.h
--- a/Engine/Source/Editor/CurveEditor/Public/SCurveEditorView.h
+++ b/Engine/Source/Editor/CurveEditor/Public/SCurveEditorView.h
@@ -25,6 +25,7 @@ class FText;
 class SCurveEditorPanel;
 struct FCurveDrawParams;
 class FCurveModel;
+class SRetainerWidget;
 
 /**
  * This is the base widget type for all views that exist on a curve editor panel. A view may contain 0 or more curves (stored in CurveInfoByID).
@@ -145,6 +146,9 @@ public:
 	 */
 	void ZoomAround(const FVector2D& Amount, double InputOrigin, double OutputOrigin);
 
+	/** This should be called every tick by an owning widget, to see if the cache is valid, which will then recreate it and invalidate widget*/
+	void CheckCacheAndInvalidateIfNeeded();
+
 public:
 
 	/**
@@ -213,8 +217,8 @@ protected:
 	void GetCurveDrawParam(TSharedPtr<FCurveEditor>& CurveEditor, const FCurveModelID& ModelID, FCurveModel* CurveModel,
 		double InputMin, double InputMax, FCurveDrawParams& OutDrawParam) const;
 
+
 	// ~SWidget interface
-	virtual void Tick(const FGeometry& AllottedGeometry, const double InCurrentTime, const float InDeltaTime) override;
 	virtual FVector2D ComputeDesiredSize(float LayoutScaleMultiplier) const override;
 
 private:
@@ -311,4 +315,13 @@ protected:
 	};
 
 	FCachedValuesToCheck CachedValues;
+
+	/** Possible pointer to a retainer widget that we may need to force update*/
+	TSharedPtr<SRetainerWidget> RetainerWidget;
+
+public:
+	void SetRetainerWidget(TSharedPtr<SRetainerWidget>& InWidget)
+	{
+		RetainerWidget = InWidget;
+	}
 };
\ No newline at end of file
