Sequencer: Curves: Remove ensure and just do valid check to avoid crash, it will eventually be valid
#trivial

diff --git a/Engine/Source/Editor/CurveEditor/Private/SCurveEditorView.cpp b/Engine/Source/Editor/CurveEditor/Private/SCurveEditorView.cpp
--- a/Engine/Source/Editor/CurveEditor/Private/SCurveEditorView.cpp
+++ b/Engine/Source/Editor/CurveEditor/Private/SCurveEditorView.cpp
@@ -322,7 +322,10 @@ void SCurveEditorView::GetCurveDrawParam(TSharedPtr<FCurveEditor>& CurveEditor,c
 void SCurveEditorView::CheckCacheAndInvalidateIfNeeded()
 {
 	TSharedPtr<FCurveEditor> CurveEditor = WeakCurveEditor.Pin();
-	check(CurveEditor.IsValid());
+	if (CurveEditor.IsValid() == false)
+	{
+		return;
+	}
 	const bool bUseCurveCache = CVarUseCurveCache.GetValueOnGameThread();
 
 	if (bUseCurveCache)
