CIS fix for not allowing forward declare. Note we can't use the default constructor/deconstructor defined in the cpp file since it relies upon a class with a private constructor.

#trivial
preflight seems broken but run nonunity tests locally

diff --git a/Engine/Source/Editor/CurveEditor/Public/SCurveEditorView.h b/Engine/Source/Editor/CurveEditor/Public/SCurveEditorView.h
--- a/Engine/Source/Editor/CurveEditor/Public/SCurveEditorView.h
+++ b/Engine/Source/Editor/CurveEditor/Public/SCurveEditorView.h
@@ -18,12 +18,12 @@
 #include "Widgets/DeclarativeSyntaxSupport.h"
 #include "Widgets/SCompoundWidget.h"
 #include "CurveEditorSettings.h"
+#include "CurveDrawInfo.h"
 
 class FCurveEditor;
 class FSlateRect;
 class FText;
 class SCurveEditorPanel;
-struct FCurveDrawParams;
 class FCurveModel;
 class SRetainerWidget;
 
