Anim Sequence: Curve View: Revert changen for cached curves done to the the Curve panel.
#trivial

diff --git a/Engine/Source/Editor/CurveEditor/Private/SCurveViewerPanel.cpp b/Engine/Source/Editor/CurveEditor/Private/SCurveViewerPanel.cpp
--- a/Engine/Source/Editor/CurveEditor/Private/SCurveViewerPanel.cpp
+++ b/Engine/Source/Editor/CurveEditor/Private/SCurveViewerPanel.cpp
@@ -48,6 +48,12 @@ void SCurveViewerPanel::Construct(const FArguments& InArgs, TSharedRef<FCurveEdi
 	SetClipping(EWidgetClipping::ClipToBounds);
 }
 
+void SCurveViewerPanel::Tick(const FGeometry& AllottedGeometry, const double InCurrentTime, const float InDeltaTime)
+{
+	CachedDrawParams.Reset();
+	GetCurveDrawParams(CachedDrawParams);
+}
+
 int32 SCurveViewerPanel::OnPaint(const FPaintArgs& Args, const FGeometry& AllottedGeometry, const FSlateRect& MyCullingRect, FSlateWindowElementList& OutDrawElements, int32 LayerId, const FWidgetStyle& InWidgetStyle, bool bParentEnabled) const
 {
 	const ESlateDrawEffect DrawEffects = ShouldBeEnabled(bParentEnabled) ? ESlateDrawEffect::None : ESlateDrawEffect::DisabledEffect;
diff --git a/Engine/Source/Editor/CurveEditor/Public/SCurveViewerPanel.h b/Engine/Source/Editor/CurveEditor/Public/SCurveViewerPanel.h
--- a/Engine/Source/Editor/CurveEditor/Public/SCurveViewerPanel.h
+++ b/Engine/Source/Editor/CurveEditor/Public/SCurveViewerPanel.h
@@ -48,6 +48,7 @@ class CURVEEDITOR_API SCurveViewerPanel : public SCurveEditorView
 private:
 
 	virtual int32 OnPaint(const FPaintArgs& Args, const FGeometry& AllottedGeometry, const FSlateRect& MyCullingRect, FSlateWindowElementList& OutDrawElements, int32 LayerId, const FWidgetStyle& InWidgetStyle, bool bParentEnabled) const override;
+	virtual void Tick(const FGeometry& AllottedGeometry, const double InCurrentTime, const float InDeltaTime) override;
 
 public:
 
