Sequencer: Curve Editor: Fix Shift Drag mechanics so it locks immediately and reliably by using the last position instead of current position when it starts.  Didn't add the last position to the ICurveDragOperation base class since most of them don't use it, only for those that snap.

diff --git a/Engine/Plugins/Editor/CurveEditorTools/Source/CurveEditorTools/Private/Tools/CurveEditorTransformTool.cpp b/Engine/Plugins/Editor/CurveEditorTools/Source/CurveEditorTools/Private/Tools/CurveEditorTransformTool.cpp
--- a/Engine/Plugins/Editor/CurveEditorTools/Source/CurveEditorTools/Private/Tools/CurveEditorTransformTool.cpp
+++ b/Engine/Plugins/Editor/CurveEditorTools/Source/CurveEditorTools/Private/Tools/CurveEditorTransformTool.cpp
@@ -391,7 +391,7 @@ FReply FCurveEditorTransformTool::OnMouseMove(TSharedRef<SWidget> OwningWidget,
 		}
 		else if (DelayedDrag->AttemptDragStart(MouseEvent))
 		{
-			InitialMousePosition = MouseEvent.GetScreenSpacePosition();
+			LastMousePosition = InitialMousePosition = MouseEvent.GetScreenSpacePosition();
 			OnDragStart();
 
 			// Steal the capture, as we're now the authoritative widget in charge of a mouse-drag operation
@@ -983,7 +983,7 @@ void FCurveEditorTransformTool::OnDrag(const FPointerEvent& InMouseEvent, const
 		// Dragging the center is the easy case!
 		if (TransformWidget.SelectedAnchorFlags == ECurveEditorAnchorFlags::Center)
 		{
-			const FVector2D AxisLockedMousePosition = CurveEditor->GetAxisSnap().GetSnappedPosition(InitialMousePosition, InMouseEvent.GetScreenSpacePosition(), InMouseEvent, SnappingState);
+			const FVector2D AxisLockedMousePosition = CurveEditor->GetAxisSnap().GetSnappedPosition(InitialMousePosition, LastMousePosition, InMouseEvent.GetScreenSpacePosition(), InMouseEvent, SnappingState);
 			{
 				FVector2D MouseDelta = AxisLockedMousePosition - InitialMousePosition;
 				TransformWidget.Position = TransformWidget.StartPosition + MouseDelta;
@@ -1061,7 +1061,7 @@ void FCurveEditorTransformTool::OnDrag(const FPointerEvent& InMouseEvent, const
 			}
 		}
 	}
-
+	LastMousePosition = InMouseEvent.GetScreenSpacePosition();
 	UpdateMarqueeBoundingBox();
 
 	UpdateToolOptions();
diff --git a/Engine/Plugins/Editor/CurveEditorTools/Source/CurveEditorTools/Private/Tools/CurveEditorTransformTool.h b/Engine/Plugins/Editor/CurveEditorTools/Source/CurveEditorTools/Private/Tools/CurveEditorTransformTool.h
--- a/Engine/Plugins/Editor/CurveEditorTools/Source/CurveEditorTools/Private/Tools/CurveEditorTransformTool.h
+++ b/Engine/Plugins/Editor/CurveEditorTools/Source/CurveEditorTools/Private/Tools/CurveEditorTransformTool.h
@@ -184,6 +184,7 @@ private:
 	TRange<double> InputMinMax;
 
 	FVector2D InitialMousePosition;
+	FVector2D LastMousePosition;
 	FCurveEditorAxisSnap::FSnapState SnappingState;
 
 	/** UStruct that displays properties to be modified on screen */
diff --git a/Engine/Source/Editor/CurveEditor/Private/DragOperations/CurveEditorDragOperation_MoveKeys.cpp b/Engine/Source/Editor/CurveEditor/Private/DragOperations/CurveEditorDragOperation_MoveKeys.cpp
--- a/Engine/Source/Editor/CurveEditor/Private/DragOperations/CurveEditorDragOperation_MoveKeys.cpp
+++ b/Engine/Source/Editor/CurveEditor/Private/DragOperations/CurveEditorDragOperation_MoveKeys.cpp
@@ -38,6 +38,7 @@ void FCurveEditorDragOperation_MoveKeys::OnBeginDrag(FVector2D InitialPosition,
 	KeysByCurve.Reset();
 	CurveEditor->SuppressBoundTransformUpdates(true);
 
+	LastMousePosition = CurrentPosition;
 	for (const TTuple<FCurveModelID, FKeyHandleSet>& Pair : CurveEditor->GetSelection().GetAll())
 	{
 		FCurveModelID CurveID = Pair.Key;
@@ -64,7 +65,8 @@ void FCurveEditorDragOperation_MoveKeys::OnBeginDrag(FVector2D InitialPosition,
 void FCurveEditorDragOperation_MoveKeys::OnDrag(FVector2D InitialPosition, FVector2D CurrentPosition, const FPointerEvent& MouseEvent)
 {
 	TArray<FKeyPosition> NewKeyPositionScratch;
-	FVector2D MousePosition = CurveEditor->GetAxisSnap().GetSnappedPosition(InitialPosition, CurrentPosition, MouseEvent, SnappingState);
+	FVector2D MousePosition = CurveEditor->GetAxisSnap().GetSnappedPosition(InitialPosition,CurrentPosition, LastMousePosition, MouseEvent, SnappingState);
+	LastMousePosition = CurrentPosition;
 
 	for (FKeyData& KeyData : KeysByCurve)
 	{
diff --git a/Engine/Source/Editor/CurveEditor/Private/DragOperations/CurveEditorDragOperation_MoveKeys.h b/Engine/Source/Editor/CurveEditor/Private/DragOperations/CurveEditorDragOperation_MoveKeys.h
--- a/Engine/Source/Editor/CurveEditor/Private/DragOperations/CurveEditorDragOperation_MoveKeys.h
+++ b/Engine/Source/Editor/CurveEditor/Private/DragOperations/CurveEditorDragOperation_MoveKeys.h
@@ -29,6 +29,7 @@ private:
 	FCurveEditor* CurveEditor;
 	TOptional<FCurvePointHandle> CardinalPoint;
 
+	FVector2D LastMousePosition;
 private:
 
 	struct FKeyData
diff --git a/Engine/Source/Editor/CurveEditor/Private/DragOperations/CurveEditorDragOperation_Pan.cpp b/Engine/Source/Editor/CurveEditor/Private/DragOperations/CurveEditorDragOperation_Pan.cpp
--- a/Engine/Source/Editor/CurveEditor/Private/DragOperations/CurveEditorDragOperation_Pan.cpp
+++ b/Engine/Source/Editor/CurveEditor/Private/DragOperations/CurveEditorDragOperation_Pan.cpp
@@ -32,7 +32,8 @@ void FCurveEditorDragOperation_PanView::OnBeginDrag(FVector2D InitialPosition, F
 
 void FCurveEditorDragOperation_PanView::OnDrag(FVector2D InitialPosition, FVector2D CurrentPosition, const FPointerEvent& MouseEvent)
 {
-	FVector2D PixelDelta = CurveEditor->GetAxisSnap().GetSnappedPosition(InitialPosition, CurrentPosition, MouseEvent, SnappingState, true) - InitialPosition;
+	FVector2D PixelDelta = CurveEditor->GetAxisSnap().GetSnappedPosition(InitialPosition, LastMousePosition,CurrentPosition, MouseEvent, SnappingState, true) - InitialPosition;
+	LastMousePosition = CurrentPosition;
 
 	FCurveEditorScreenSpace ViewSpace = View->GetViewSpace();
 
@@ -80,11 +81,14 @@ void FCurveEditorDragOperation_PanInput::OnBeginDrag(FVector2D InitialPosition,
 	InitialInputMin = InputSpace.GetInputMin();
 	InitialInputMax = InputSpace.GetInputMax();
 	SnappingState.Reset();
+	LastMousePosition = CurrentPosition;
+
 }
 
 void FCurveEditorDragOperation_PanInput::OnDrag(FVector2D InitialPosition, FVector2D CurrentPosition, const FPointerEvent& MouseEvent)
 {
-	FVector2D PixelDelta = CurveEditor->GetAxisSnap().GetSnappedPosition(InitialPosition, CurrentPosition, MouseEvent, SnappingState, true) - InitialPosition;
+	FVector2D PixelDelta = CurveEditor->GetAxisSnap().GetSnappedPosition(InitialPosition, LastMousePosition, CurrentPosition, MouseEvent, SnappingState, true) - InitialPosition;
+	LastMousePosition = CurrentPosition;
 
 	FCurveEditorScreenSpaceH InputSpace = CurveEditor->GetPanelInputSpace();
 
diff --git a/Engine/Source/Editor/CurveEditor/Private/DragOperations/CurveEditorDragOperation_Pan.h b/Engine/Source/Editor/CurveEditor/Private/DragOperations/CurveEditorDragOperation_Pan.h
--- a/Engine/Source/Editor/CurveEditor/Private/DragOperations/CurveEditorDragOperation_Pan.h
+++ b/Engine/Source/Editor/CurveEditor/Private/DragOperations/CurveEditorDragOperation_Pan.h
@@ -31,6 +31,8 @@ private:
 	double InitialInputMin, InitialInputMax;
 	double InitialOutputMin, InitialOutputMax;
 	FCurveEditorAxisSnap::FSnapState SnappingState;
+	FVector2D LastMousePosition;
+
 };
 
 class FCurveEditorDragOperation_PanInput : public ICurveEditorDragOperation
@@ -46,4 +48,5 @@ private:
 	FCurveEditor* CurveEditor;
 	double InitialInputMin, InitialInputMax;
 	FCurveEditorAxisSnap::FSnapState SnappingState;
+	FVector2D LastMousePosition;
 };
\ No newline at end of file
diff --git a/Engine/Source/Editor/CurveEditor/Public/CurveEditorSnapMetrics.h b/Engine/Source/Editor/CurveEditor/Public/CurveEditorSnapMetrics.h
--- a/Engine/Source/Editor/CurveEditor/Public/CurveEditorSnapMetrics.h
+++ b/Engine/Source/Editor/CurveEditor/Public/CurveEditorSnapMetrics.h
@@ -81,7 +81,7 @@ struct FCurveEditorAxisSnap
 		RestrictedAxisList = EAxisList::None;
 	}
 
-	FVector2D GetSnappedPosition(FVector2D InitialPosition, FVector2D CurrentPosition, const FPointerEvent& MouseEvent, FSnapState& InOutSnapState, const bool bIgnoreAxisLock = false)
+	FVector2D GetSnappedPosition(FVector2D InitialPosition, FVector2D LastPosition, FVector2D CurrentPosition, const FPointerEvent& MouseEvent, FSnapState& InOutSnapState, const bool bIgnoreAxisLock = false)
 	{
 		check(RestrictedAxisList == EAxisList::Type::None || RestrictedAxisList == EAxisList::Type::X || RestrictedAxisList == EAxisList::Type::Y);
 
@@ -98,7 +98,7 @@ struct FCurveEditorAxisSnap
 			{
 				if (!InOutSnapState.bHasStartPosition)
 				{
-					InOutSnapState.MousePosOnShiftStart = CurrentPosition;
+					InOutSnapState.MousePosOnShiftStart = LastPosition;
 					InOutSnapState.bHasStartPosition = true;
 				}
 				// If they have passed the threshold they should have a lock vector they're snapped to.
@@ -108,7 +108,7 @@ struct FCurveEditorAxisSnap
 					
 					// They have not passed the threshold yet, let's see if they've passed it now.
 					FVector2D DragDelta = CurrentPosition - InOutSnapState.MousePosOnShiftStart;
-					if (DragDelta.Size() > 2.f)
+					if (DragDelta.Size() > 0.001f)
 					{
 						InOutSnapState.bHasPassedThreshold = true;
 						InOutSnapState.MouseLockVector = FVector2D::UnitVector;
