Sequencer: Curve Editor:  Due to delayed drag the initial position of the marquee drag is wrong so we store that and end before the end drag happens, also fix to hit test tangent lines when marquee dragging

diff --git a/Engine/Source/Editor/CurveEditor/Private/DragOperations/CurveEditorDragOperation_Marquee.cpp b/Engine/Source/Editor/CurveEditor/Private/DragOperations/CurveEditorDragOperation_Marquee.cpp
--- a/Engine/Source/Editor/CurveEditor/Private/DragOperations/CurveEditorDragOperation_Marquee.cpp
+++ b/Engine/Source/Editor/CurveEditor/Private/DragOperations/CurveEditorDragOperation_Marquee.cpp
@@ -35,33 +35,28 @@ FCurveEditorDragOperation_Marquee::FCurveEditorDragOperation_Marquee(FCurveEdito
 
 void FCurveEditorDragOperation_Marquee::OnBeginDrag(FVector2D InitialPosition, FVector2D CurrentPosition, const FPointerEvent& MouseEvent)
 {
+	RealInitialPosition = CurrentPosition;
+
 	Marquee = FSlateRect(
-		FMath::Min(InitialPosition.X, CurrentPosition.X),
-		FMath::Min(InitialPosition.Y, CurrentPosition.Y),
-		FMath::Max(InitialPosition.X, CurrentPosition.X),
-		FMath::Max(InitialPosition.Y, CurrentPosition.Y)
+		FMath::Min(RealInitialPosition.X, CurrentPosition.X),
+		FMath::Min(RealInitialPosition.Y, CurrentPosition.Y),
+		FMath::Max(RealInitialPosition.X, CurrentPosition.X),
+		FMath::Max(RealInitialPosition.Y, CurrentPosition.Y)
 		);
 }
 
 void FCurveEditorDragOperation_Marquee::OnDrag(FVector2D InitialPosition, FVector2D CurrentPosition, const FPointerEvent& MouseEvent)
 {
 	Marquee = FSlateRect(
-		FMath::Min(InitialPosition.X, CurrentPosition.X),
-		FMath::Min(InitialPosition.Y, CurrentPosition.Y),
-		FMath::Max(InitialPosition.X, CurrentPosition.X),
-		FMath::Max(InitialPosition.Y, CurrentPosition.Y)
+		FMath::Min(RealInitialPosition.X, CurrentPosition.X),
+		FMath::Min(RealInitialPosition.Y, CurrentPosition.Y),
+		FMath::Max(RealInitialPosition.X, CurrentPosition.X),
+		FMath::Max(RealInitialPosition.Y, CurrentPosition.Y)
 		);
 }
 
 void FCurveEditorDragOperation_Marquee::OnEndDrag(FVector2D InitialPosition, FVector2D CurrentPosition, const FPointerEvent& MouseEvent)
 {
-	Marquee = FSlateRect(
-		FMath::Min(InitialPosition.X, CurrentPosition.X),
-		FMath::Min(InitialPosition.Y, CurrentPosition.Y),
-		FMath::Max(InitialPosition.X, CurrentPosition.X),
-		FMath::Max(InitialPosition.Y, CurrentPosition.Y)
-		);
-
 	TArray<FCurvePointHandle> AllPoints;
 
 	if (LockedToView)
diff --git a/Engine/Source/Editor/CurveEditor/Private/DragOperations/CurveEditorDragOperation_Marquee.h b/Engine/Source/Editor/CurveEditor/Private/DragOperations/CurveEditorDragOperation_Marquee.h
--- a/Engine/Source/Editor/CurveEditor/Private/DragOperations/CurveEditorDragOperation_Marquee.h
+++ b/Engine/Source/Editor/CurveEditor/Private/DragOperations/CurveEditorDragOperation_Marquee.h
@@ -34,4 +34,6 @@ private:
 	FCurveEditor* CurveEditor;
 	/** When valid, marquee selection should only occur inside this view; all geometries are in local space */
 	SCurveEditorView* LockedToView;
+	/** Real Initial Position do to the delayed drag*/
+	FVector2D RealInitialPosition;
 };
\ No newline at end of file
diff --git a/Engine/Source/Editor/CurveEditor/Private/Views/SInteractiveCurveEditorView.cpp b/Engine/Source/Editor/CurveEditor/Private/Views/SInteractiveCurveEditorView.cpp
--- a/Engine/Source/Editor/CurveEditor/Private/Views/SInteractiveCurveEditorView.cpp
+++ b/Engine/Source/Editor/CurveEditor/Private/Views/SInteractiveCurveEditorView.cpp
@@ -571,8 +571,12 @@ bool SInteractiveCurveEditorView::GetPointsWithinWidgetRange(const FSlateRect& W
 		return false;
 	}
 
+	FVector2D LinePoints[2];
+	FVector Start, End, StartToEnd;
+	FBox WidgetRectangleBox(FVector(WidgetRectangle.Left, WidgetRectangle.Top, 0), FVector(WidgetRectangle.Right, WidgetRectangle.Bottom, 0));
+
 	// Iterate through all of our points and see which points the marquee overlaps. Both of these coordinate systems
-	// are in screen space pixels.
+	// are in screen space pixels.  Also check tangent lines
 	bool bFound = false;
 	for (const FCurveDrawParams& DrawParams : CachedDrawParams)
 	{
@@ -588,6 +592,21 @@ bool SInteractiveCurveEditorView::GetPointsWithinWidgetRange(const FSlateRect& W
 				OutPoints->Add(FCurvePointHandle(DrawParams.GetID(), Point.Type, Point.KeyHandle));
 				bFound = true;
 			}
+			else if (Point.LineDelta.X != 0.f || Point.LineDelta.Y != 0.f) //if tangent hit test line
+			{
+				LinePoints[0] = Point.ScreenPosition + Point.LineDelta.GetSafeNormal() * (DrawInfo.ScreenSize.X * .5f);
+				LinePoints[1] = Point.ScreenPosition + Point.LineDelta;
+
+				Start = FVector(LinePoints[0].X, LinePoints[0].Y, 0);
+				End = FVector(LinePoints[1].X, LinePoints[1].Y, 0);
+				StartToEnd = End - Start;
+
+				if (FMath::LineBoxIntersection(WidgetRectangleBox, Start, End, StartToEnd))
+				{
+					OutPoints->Add(FCurvePointHandle(DrawParams.GetID(), Point.Type, Point.KeyHandle));
+					bFound = true;
+				}
+			}
 		}
 	}
 
