Sequencer: Curve Editor: Keep Output bounds the same even if curves are destroyed so users can  select/deselect and keep their curve bounds.

diff --git a/Engine/Source/Editor/CurveEditor/Private/CurveEditor.cpp b/Engine/Source/Editor/CurveEditor/Private/CurveEditor.cpp
--- a/Engine/Source/Editor/CurveEditor/Private/CurveEditor.cpp
+++ b/Engine/Source/Editor/CurveEditor/Private/CurveEditor.cpp
@@ -207,7 +207,14 @@ FCurveModelID FCurveEditor::AddCurveForTreeItem(TUniquePtr<FCurveModel>&& InCurv
 
 	return NewID;
 }
-
+void FCurveEditor::ResetMinMaxes()
+{
+	TSharedPtr<SCurveEditorPanel> Panel = WeakPanel.Pin();
+	if (Panel.IsValid())
+	{
+		Panel->ResetMinMaxes();
+	}
+}
 void FCurveEditor::RemoveCurve(FCurveModelID InCurveID)
 {
 	TSharedPtr<SCurveEditorPanel> Panel = WeakPanel.Pin();
diff --git a/Engine/Source/Editor/CurveEditor/Private/SCurveEditorPanel.cpp b/Engine/Source/Editor/CurveEditor/Private/SCurveEditorPanel.cpp
--- a/Engine/Source/Editor/CurveEditor/Private/SCurveEditorPanel.cpp
+++ b/Engine/Source/Editor/CurveEditor/Private/SCurveEditorPanel.cpp
@@ -491,7 +491,6 @@ void SCurveEditorPanel::Tick(const FGeometry& AllottedGeometry, const double InC
 		{
 			CurveEditor->ZoomToFitCurves(CurveEditor->GetEditedCurves().Array());
 		}
-
 		bNeedsRefresh = false;
 		CachedActiveCurvesSerialNumber = CurveEditor->GetActiveCurvesSerialNumber();
 		bWasRefreshed = true;
@@ -504,11 +503,35 @@ void SCurveEditorPanel::Tick(const FGeometry& AllottedGeometry, const double InC
 	CachedSelectionSerialNumber = CurveEditor->Selection.GetSerialNumber();
 }
 
+void SCurveEditorPanel::ResetMinMaxes()
+{
+	//only reset the min/max if we have views since we will then get these values from them
+	//otherwise if we didn't we would end up with everything back to 0,1 again.
+	if (CurveViews.IsEmpty() == false)
+	{
+		LastOutputMin = DBL_MAX;
+		LastOutputMax = DBL_MIN;
+	}
+}
+
 void SCurveEditorPanel::RemoveCurveFromViews(FCurveModelID InCurveID)
 {
 	for (auto It = CurveViews.CreateKeyIterator(InCurveID); It; ++It)
 	{
-		FCurveEditorPanelViewTracker::RemoveCurveFromView(&It.Value().Get(), InCurveID);
+		SCurveEditorView* View = &It.Value().Get();
+		//cache these so we can re-use it on reconstruction
+		if (View)
+		{
+			if (View->GetOutputMin() < LastOutputMin)
+			{
+				LastOutputMin = View->GetOutputMin();
+			}
+			if (View->GetOutputMax() > LastOutputMax)
+			{
+				LastOutputMax = View->GetOutputMax();
+			}
+		}
+		FCurveEditorPanelViewTracker::RemoveCurveFromView(View, InCurveID);
 		It.RemoveCurrent();
 	}
 }
@@ -552,6 +575,10 @@ TSharedPtr<SCurveEditorView> SCurveEditorPanel::CreateViewOfType(FCurveModelID C
 			{
 				It.RemoveCurrent();
 			}
+			if (LastOutputMin != DBL_MAX && LastOutputMax != DBL_MIN)
+			{
+				View->SetOutputBounds(LastOutputMin, LastOutputMax);
+			}
 
 			return View;
 		}
diff --git a/Engine/Source/Editor/CurveEditor/Private/Tree/CurveEditorTree.cpp b/Engine/Source/Editor/CurveEditor/Private/Tree/CurveEditorTree.cpp
--- a/Engine/Source/Editor/CurveEditor/Private/Tree/CurveEditorTree.cpp
+++ b/Engine/Source/Editor/CurveEditor/Private/Tree/CurveEditorTree.cpp
@@ -492,7 +492,7 @@ void FCurveEditorTree::SetDirectSelection(TArray<FCurveEditorTreeItemID>&& TreeI
 			}
 		}
 	}
-
+	InCurveEditor->ResetMinMaxes();
 	for (TTuple<FCurveEditorTreeItemID, ECurveEditorTreeSelectionState> OldItem : PreviousSelection)
 	{
 		const ECurveEditorTreeSelectionState* NewState = Selection.Find(OldItem.Key);
diff --git a/Engine/Source/Editor/CurveEditor/Public/CurveEditor.h b/Engine/Source/Editor/CurveEditor/Public/CurveEditor.h
--- a/Engine/Source/Editor/CurveEditor/Public/CurveEditor.h
+++ b/Engine/Source/Editor/CurveEditor/Public/CurveEditor.h
@@ -133,6 +133,7 @@ public:
 
 	FCurveEditorScreenSpaceH GetPanelInputSpace() const;
 
+	void ResetMinMaxes();
 public:
 	/**
 	 * Zoom the curve editor to fit all the selected curves (or all curves if none selected)
diff --git a/Engine/Source/Editor/CurveEditor/Public/SCurveEditorPanel.h b/Engine/Source/Editor/CurveEditor/Public/SCurveEditorPanel.h
--- a/Engine/Source/Editor/CurveEditor/Public/SCurveEditorPanel.h
+++ b/Engine/Source/Editor/CurveEditor/Public/SCurveEditorPanel.h
@@ -148,6 +148,8 @@ class CURVEEDITOR_API SCurveEditorPanel : public SCompoundWidget
 	/** Undo occurred, invalidate or update internal structures */
 	void PostUndo();
 
+	/** Reset Stored Min/Max's*/
+	void ResetMinMaxes();
 private:
 	// SWidget Interface
 	virtual void Tick(const FGeometry& AllottedGeometry, const double InCurrentTime, const float InDeltaTime) override;
@@ -266,7 +268,7 @@ private:
 	void OnColumnFillCoefficientChanged(float FillCoefficient, int32 ColumnIndex);
 
 	void OnSplitterFinishedResizing();
-
+	
 private:
 
 	/**
@@ -337,6 +339,10 @@ private:
 	/** Reconstructs the properties widget on tool switch */
 	void OnCurveEditorToolChanged(FCurveEditorToolID InToolId);
 
+	/** Last Output Min and Max values for the views*/
+	double LastOutputMin = DBL_MAX;
+	double LastOutputMax = DBL_MIN;
+
 	/** The last set View Mode for this UI. */
 	ECurveEditorViewID DefaultViewID;
 
