Fixed issue where color curve editor wouldn't update when gradient stops were moved

diff --git a/Engine/Source/Editor/CurveAssetEditor/Private/SColorGradientCurveEditorView.cpp b/Engine/Source/Editor/CurveAssetEditor/Private/SColorGradientCurveEditorView.cpp
--- a/Engine/Source/Editor/CurveAssetEditor/Private/SColorGradientCurveEditorView.cpp
+++ b/Engine/Source/Editor/CurveAssetEditor/Private/SColorGradientCurveEditorView.cpp
@@ -19,6 +19,7 @@ void SColorGradientCurveEditorView::Construct(const FArguments& InArgs, TSharedR
 	bInteractive = 0;
 	bAutoSize = 1;
 	bAllowEmpty = 1;
+	WeakCurveEditor = InCurveEditor;
 
 	ChildSlot
 	[
@@ -33,4 +34,11 @@ void SColorGradientCurveEditorView::Construct(const FArguments& InArgs, TSharedR
 			.IsEditingEnabled(InArgs._IsEditingEnabled)
 		]
 	];
+}
+
+void SColorGradientCurveEditorView::CheckCacheAndInvalidateIfNeeded()
+{
+	// Always refresh for now
+	// Could be improved by combining SCurveEditorView::CheckCacheAndInvalidateIfNeeded with a separate check for changes in the gradient stops
+	RefreshRetainer();
 }
\ No newline at end of file
diff --git a/Engine/Source/Editor/CurveAssetEditor/Private/SColorGradientCurveEditorView.h b/Engine/Source/Editor/CurveAssetEditor/Private/SColorGradientCurveEditorView.h
--- a/Engine/Source/Editor/CurveAssetEditor/Private/SColorGradientCurveEditorView.h
+++ b/Engine/Source/Editor/CurveAssetEditor/Private/SColorGradientCurveEditorView.h
@@ -39,6 +39,8 @@ public:
 	void Construct(const FArguments& InArgs, TSharedRef<FCurveEditor> InCurveEditor);
 
 	TSharedRef<SColorGradientEditor> GetGradientEditor() const { return GradientViewer.ToSharedRef(); }
+
+	virtual void CheckCacheAndInvalidateIfNeeded() override;
 private:
 	TSharedPtr<SColorGradientEditor> GradientViewer;
 };
diff --git a/Engine/Source/Editor/CurveEditor/Private/SCurveEditorView.cpp b/Engine/Source/Editor/CurveEditor/Private/SCurveEditorView.cpp
--- a/Engine/Source/Editor/CurveEditor/Private/SCurveEditorView.cpp
+++ b/Engine/Source/Editor/CurveEditor/Private/SCurveEditorView.cpp
@@ -319,6 +319,14 @@ void SCurveEditorView::GetCurveDrawParam(TSharedPtr<FCurveEditor>& CurveEditor,c
 	}
 }
 
+void SCurveEditorView::RefreshRetainer()
+{
+	if (RetainerWidget)
+	{
+		RetainerWidget->RequestRender();
+	}
+}
+
 void SCurveEditorView::CheckCacheAndInvalidateIfNeeded()
 {
 	TSharedPtr<FCurveEditor> CurveEditor = WeakCurveEditor.Pin();
@@ -378,10 +386,7 @@ void SCurveEditorView::CheckCacheAndInvalidateIfNeeded()
 			CachedDrawParams.Reset();
 			GetCurveDrawParams(CachedDrawParams);
 			CurveCacheFlags = ECurveCacheFlags::CheckCurves;
-			if (RetainerWidget)
-			{
-				RetainerWidget->RequestRender();
-			}
+			RefreshRetainer();
 		}
 		else if (CurveCacheFlags == ECurveCacheFlags::CheckCurves)
 		{
@@ -395,9 +400,9 @@ void SCurveEditorView::CheckCacheAndInvalidateIfNeeded()
 					GetCurveDrawParam(CurveEditor, Params.GetID(), CurveModel, CachedValues.CachedInputMin, CachedValues.CachedInputMax, Params);
 				}
 			}
-			if (bSomethingChanged && RetainerWidget)
+			if (bSomethingChanged)
 			{
-				RetainerWidget->RequestRender();
+				RefreshRetainer();
 			}
 		}
 	}
@@ -411,9 +416,6 @@ void SCurveEditorView::CheckCacheAndInvalidateIfNeeded()
 		CachedDrawParams.Reset();
 		GetCurveDrawParams(CachedDrawParams);
 
-		if (RetainerWidget)
-		{
-			RetainerWidget->RequestRender();
-		}
+		RefreshRetainer();
 	}
 }
\ No newline at end of file
diff --git a/Engine/Source/Editor/CurveEditor/Public/SCurveEditorView.h b/Engine/Source/Editor/CurveEditor/Public/SCurveEditorView.h
--- a/Engine/Source/Editor/CurveEditor/Public/SCurveEditorView.h
+++ b/Engine/Source/Editor/CurveEditor/Public/SCurveEditorView.h
@@ -147,7 +147,7 @@ public:
 	void ZoomAround(const FVector2D& Amount, double InputOrigin, double OutputOrigin);
 
 	/** This should be called every tick by an owning widget, to see if the cache is valid, which will then recreate it and invalidate widget*/
-	void CheckCacheAndInvalidateIfNeeded();
+	virtual void CheckCacheAndInvalidateIfNeeded();
 
 public:
 
@@ -217,6 +217,8 @@ protected:
 	void GetCurveDrawParam(TSharedPtr<FCurveEditor>& CurveEditor, const FCurveModelID& ModelID, FCurveModel* CurveModel,
 		double InputMin, double InputMax, FCurveDrawParams& OutDrawParam) const;
 
+	/** Request a new render from the retainer widget */
+	void RefreshRetainer();
 
 	// ~SWidget interface
 	virtual FVector2D ComputeDesiredSize(float LayoutScaleMultiplier) const override;
