Fixed mismatched copy constructor/assignment operator in FCurveModelID.

diff --git a/Engine/Source/Editor/CurveEditor/Public/CurveEditorTypes.h b/Engine/Source/Editor/CurveEditor/Public/CurveEditorTypes.h
--- a/Engine/Source/Editor/CurveEditor/Public/CurveEditorTypes.h
+++ b/Engine/Source/Editor/CurveEditor/Public/CurveEditorTypes.h
@@ -84,11 +84,6 @@ struct CURVEEDITOR_API FCurveModelID
 		return GetTypeHash(In.ID);
 	}
 
-	FCurveModelID(const FCurveModelID& InOther)
-		: ID(InOther.ID)
-	{
-	}
-
 private:
 	FCurveModelID() {}
 
