Fixed mismatched copy constructor/assignment operator for FCurveEditorToolID.

diff --git a/Engine/Source/Editor/CurveEditor/Public/ICurveEditorToolExtension.h b/Engine/Source/Editor/CurveEditor/Public/ICurveEditorToolExtension.h
--- a/Engine/Source/Editor/CurveEditor/Public/ICurveEditorToolExtension.h
+++ b/Engine/Source/Editor/CurveEditor/Public/ICurveEditorToolExtension.h
@@ -63,11 +63,6 @@ struct CURVEEDITOR_API FCurveEditorToolID
 		return GetTypeHash(In.ID);
 	}
 
-	FCurveEditorToolID(const FCurveEditorToolID& InOther)
-		: ID(InOther.ID)
-	{
-	}
-
 private:
 	FCurveEditorToolID() {}
 
