Sequencer: Curve Editor: Make normalized mode scrollable and panable. Also make sure that if you frame to the normalized mode now it doesn't use the max/min values but defaults.

diff --git a/Engine/Source/Editor/CurveEditor/Private/CurveEditor.cpp b/Engine/Source/Editor/CurveEditor/Private/CurveEditor.cpp
--- a/Engine/Source/Editor/CurveEditor/Private/CurveEditor.cpp
+++ b/Engine/Source/Editor/CurveEditor/Private/CurveEditor.cpp
@@ -712,13 +712,14 @@ void FCurveEditor::ZoomToFitInternal(EAxisList::Type Axes, const TMap<FCurveMode
 	for (const TTuple<TSharedRef<SCurveEditorView>, TTuple<double, double>>& ViewAndBounds : ViewToOutputBounds)
 	{
 		TSharedRef<SCurveEditorView> View = ViewAndBounds.Key;
+
 		double OutputMin = ViewAndBounds.Value.Get<0>();
 		double OutputMax = ViewAndBounds.Value.Get<1>();
 
 		// If zooming to the same (or invalid) min/max, keep the same zoom scale and center within the timeline
 		if (OutputMin >= OutputMax)
 		{
-			const double HalfOutputScale = (View->GetOutputMax() - View->GetOutputMin())*0.5;
+			const double HalfOutputScale = (View->GetOutputMax() - View->GetOutputMin()) * 0.5;
 			OutputMin -= HalfOutputScale;
 			OutputMax += HalfOutputScale;
 		}
@@ -731,7 +732,7 @@ void FCurveEditor::ZoomToFitInternal(EAxisList::Type Axes, const TMap<FCurveMode
 			{
 				PanelHeight = WeakPanel.Pin()->GetViewContainerGeometry().GetLocalSize().Y;
 			}
-			else 
+			else
 			{
 				PanelHeight = View->GetViewSpace().GetPhysicalHeight();
 			}
@@ -744,8 +745,7 @@ void FCurveEditor::ZoomToFitInternal(EAxisList::Type Axes, const TMap<FCurveMode
 			OutputMin -= OutputPadding;
 			OutputMax = FMath::Max(OutputMin + MinOutputZoom, OutputMax) + OutputPadding;
 		}
-
-		View->SetOutputBounds(OutputMin, OutputMax);
+		View->FrameVertical(OutputMin, OutputMax);
 	}
 }
 
diff --git a/Engine/Source/Editor/CurveEditor/Private/SCurveEditorView.cpp b/Engine/Source/Editor/CurveEditor/Private/SCurveEditorView.cpp
--- a/Engine/Source/Editor/CurveEditor/Private/SCurveEditorView.cpp
+++ b/Engine/Source/Editor/CurveEditor/Private/SCurveEditorView.cpp
@@ -108,6 +108,12 @@ void SCurveEditorView::RemoveCurve(FCurveModelID CurveID)
 	}
 }
 
+void SCurveEditorView::FrameVertical(double InOutputMin, double InOutputMax)
+{
+	//default just set's output if it can
+	SetOutputBounds(InOutputMin, InOutputMax);
+}
+
 void SCurveEditorView::SetOutputBounds(double InOutputMin, double InOutputMax)
 {
 	if (!bFixedOutputBounds)
@@ -150,8 +156,11 @@ void SCurveEditorView::ZoomAround(const FVector2D& Amount, double InputOrigin, d
 	}
 }
 
-void SCurveEditorView::GetCurveDrawParams(TArray<FCurveDrawParams>& OutDrawParams) const
+void SCurveEditorView::GetCurveDrawParams(TArray<FCurveDrawParams>& OutDrawParams) 
 {
+	//make sure the transform is set up
+	UpdateViewToTransformCurves();
+
 	TSharedPtr<FCurveEditor> CurveEditor = WeakCurveEditor.Pin();
 	if (!CurveEditor)
 	{
diff --git a/Engine/Source/Editor/CurveEditor/Private/Views/SCurveEditorViewNormalized.cpp b/Engine/Source/Editor/CurveEditor/Private/Views/SCurveEditorViewNormalized.cpp
--- a/Engine/Source/Editor/CurveEditor/Private/Views/SCurveEditorViewNormalized.cpp
+++ b/Engine/Source/Editor/CurveEditor/Private/Views/SCurveEditorViewNormalized.cpp
@@ -40,9 +40,9 @@ constexpr float NormalizedPadding = 10.f;
 
 void SCurveEditorViewNormalized::Construct(const FArguments& InArgs, TWeakPtr<FCurveEditor> InCurveEditor)
 {
+	//when created we set the output bounds to be fixed since otherwise it may get fitted
 	bFixedOutputBounds = true;
-	OutputMin = 0.0;
-	OutputMax = 1.0;
+	FrameVertical(-0.1, 1.1);
 
 	SInteractiveCurveEditorView::Construct(InArgs, InCurveEditor);
 
@@ -58,24 +58,12 @@ void SCurveEditorViewNormalized::Construct(const FArguments& InArgs, TWeakPtr<FC
 	];
 }
 
-void SCurveEditorViewNormalized::GetGridLinesY(TSharedRef<const FCurveEditor> CurveEditor, TArray<float>& MajorGridLines, TArray<float>& MinorGridLines, TArray<FText>* MajorGridLabels) const
-{
-	FCurveEditorScreenSpace ViewSpace = GetViewSpace();
-
-	MajorGridLines.Add(ViewSpace.ValueToScreen(0.0));
-	MajorGridLines.Add(ViewSpace.ValueToScreen(0.5));
-	MajorGridLines.Add(ViewSpace.ValueToScreen(1.0));
-
-	MinorGridLines.Add(ViewSpace.ValueToScreen(0.25));
-	MinorGridLines.Add(ViewSpace.ValueToScreen(0.75));
-}
-
-FTransform2D CalculateViewToCurveTransform(const double InCurveOutputMin, const double InCurveOutputMax)
+FTransform2D CalculateViewToCurveTransform(const double OutputMin, const double OutputMax, const double InCurveOutputMin, const double InCurveOutputMax)
 {
 	const double Scale = (InCurveOutputMax - InCurveOutputMin);
 	if (InCurveOutputMax > InCurveOutputMin)
 	{
-		return FTransform2D(FScale2D(1.f, Scale), FVector2D(0.f, InCurveOutputMin));
+		return FTransform2D(FScale2D(1.0, Scale), FVector2D(0.0, InCurveOutputMin));
 	}
 	else
 	{
@@ -116,7 +104,7 @@ void SCurveEditorViewNormalized::DrawBufferedCurves(const FGeometry& AllottedGeo
 		FTransform2D ViewToBufferedCurveTransform;
 		double CurveOutputMin = BufferedCurve->GetValueMin(), CurveOutputMax = BufferedCurve->GetValueMax();
 
-		ViewToBufferedCurveTransform = CalculateViewToCurveTransform(CurveOutputMin, CurveOutputMax);
+		ViewToBufferedCurveTransform = CalculateViewToCurveTransform(OutputMin, OutputMax,CurveOutputMin, CurveOutputMax);
 
 		TArray<TTuple<double, double>> CurveSpaceInterpolatingPoints;
 		FCurveEditorScreenSpace CurveSpace = GetViewSpace().ToCurveSpace(ViewToBufferedCurveTransform);
@@ -159,7 +147,13 @@ void SCurveEditorViewNormalized::PaintView(const FPaintArgs& Args, const FGeomet
 	}
 }
 
-void SCurveEditorViewNormalized::Tick(const FGeometry& AllottedGeometry, const double InCurrentTime, const float InDeltaTime)
+void SCurveEditorViewNormalized::FrameVertical(double InOutputMin, double InOutputMax)
+{
+	OutputMin = -0.1;
+	OutputMax = 1.1;;
+}
+
+void SCurveEditorViewNormalized::UpdateViewToTransformCurves() 
 {
 	TSharedPtr<FCurveEditor> CurveEditor = WeakCurveEditor.Pin();
 	if (!CurveEditor)
@@ -167,27 +161,34 @@ void SCurveEditorViewNormalized::Tick(const FGeometry& AllottedGeometry, const d
 		return;
 	}
 
-	if (!CurveEditor->AreBoundTransformUpdatesSuppressed())
+	for (auto It = CurveInfoByID.CreateIterator(); It; ++It)
 	{
-		const double ValuePerPixel = 1.0 / AllottedGeometry.GetLocalSize().Y;
-		const double ValueSpacePadding = NormalizedPadding * ValuePerPixel;
+		FCurveModel* Curve = CurveEditor->FindCurve(It.Key());
+		if (!ensureAlways(Curve))
+		{
+			continue;
+		}
 
-		OutputMin = 0.0 - ValueSpacePadding;
-		OutputMax = 1.0 + ValueSpacePadding;
+		double CurveOutputMin = 0, CurveOutputMax = 1;
+		Curve->GetValueRange(CurveOutputMin, CurveOutputMax);
 
-		for (auto It = CurveInfoByID.CreateIterator(); It; ++It)
-		{
-			FCurveModel* Curve = CurveEditor->FindCurve(It.Key());
-			if (!ensureAlways(Curve))
-			{
-				continue;
-			}
+		It->Value.ViewToCurveTransform = CalculateViewToCurveTransform(OutputMin, OutputMax, CurveOutputMin, CurveOutputMax);
+	}
+}
 
-			double CurveOutputMin = 0, CurveOutputMax = 1;
-			Curve->GetValueRange(CurveOutputMin, CurveOutputMax);
 
-			It->Value.ViewToCurveTransform = CalculateViewToCurveTransform(CurveOutputMin, CurveOutputMax);
-		}
+void SCurveEditorViewNormalized::Tick(const FGeometry& AllottedGeometry, const double InCurrentTime, const float InDeltaTime)
+{
+	TSharedPtr<FCurveEditor> CurveEditor = WeakCurveEditor.Pin();
+	if (!CurveEditor)
+	{
+		return;
+	}
+
+	if (!CurveEditor->AreBoundTransformUpdatesSuppressed())
+	{
+		bFixedOutputBounds = false;
+		UpdateViewToTransformCurves();
 	}
 
 	SInteractiveCurveEditorView::Tick(AllottedGeometry, InCurrentTime, InDeltaTime);
diff --git a/Engine/Source/Editor/CurveEditor/Private/Views/SCurveEditorViewStacked.cpp b/Engine/Source/Editor/CurveEditor/Private/Views/SCurveEditorViewStacked.cpp
--- a/Engine/Source/Editor/CurveEditor/Private/Views/SCurveEditorViewStacked.cpp
+++ b/Engine/Source/Editor/CurveEditor/Private/Views/SCurveEditorViewStacked.cpp
@@ -379,38 +379,50 @@ void SCurveEditorViewStacked::DrawBufferedCurves(const FGeometry& AllottedGeomet
 	}
 }
 
-void SCurveEditorViewStacked::Tick(const FGeometry& AllottedGeometry, const double InCurrentTime, const float InDeltaTime)
+void SCurveEditorViewStacked::UpdateViewToTransformCurves()
 {
+
 	TSharedPtr<FCurveEditor> CurveEditor = WeakCurveEditor.Pin();
 	if (!CurveEditor)
 	{
 		return;
 	}
+	double ValuePerPixel = 1.0 / StackedHeight;
+	double ValueSpacePadding = StackedPadding * ValuePerPixel;
 
-	if (!CurveEditor->AreBoundTransformUpdatesSuppressed())
+	for (auto It = CurveInfoByID.CreateIterator(); It; ++It)
 	{
-		double ValuePerPixel = 1.0 / StackedHeight;
-		double ValueSpacePadding = StackedPadding * ValuePerPixel;
-
-		for (auto It = CurveInfoByID.CreateIterator(); It; ++It)
+		FCurveModel* Curve = CurveEditor->FindCurve(It.Key());
+		if (!ensureAlways(Curve))
 		{
-			FCurveModel* Curve = CurveEditor->FindCurve(It.Key());
-			if (!ensureAlways(Curve))
-			{
-				continue;
-			}
+			continue;
+		}
 
-			const int32  CurveIndexFromBottom = CurveInfoByID.Num() - It->Value.CurveIndex - 1;
-			const double PaddingToBottomOfView = (CurveIndexFromBottom + 1)*ValueSpacePadding;
-			const double ValueOffset = -CurveIndexFromBottom - PaddingToBottomOfView;
+		const int32  CurveIndexFromBottom = CurveInfoByID.Num() - It->Value.CurveIndex - 1;
+		const double PaddingToBottomOfView = (CurveIndexFromBottom + 1) * ValueSpacePadding;
+		const double ValueOffset = -CurveIndexFromBottom - PaddingToBottomOfView;
 
-			double CurveOutputMin = 0, CurveOutputMax = 1;
-			Curve->GetValueRange(CurveOutputMin, CurveOutputMax);
+		double CurveOutputMin = 0, CurveOutputMax = 1;
+		Curve->GetValueRange(CurveOutputMin, CurveOutputMax);
 
-			It->Value.ViewToCurveTransform = CalculateViewToCurveTransform(CurveOutputMin, CurveOutputMax, ValueOffset);
-		}
+		It->Value.ViewToCurveTransform = CalculateViewToCurveTransform(CurveOutputMin, CurveOutputMax, ValueOffset);
+	}
+
+	OutputMax = FMath::Max(OutputMin + CurveInfoByID.Num() + ValueSpacePadding * (CurveInfoByID.Num() + 1), 1.0);
+}
+
+
+void SCurveEditorViewStacked::Tick(const FGeometry& AllottedGeometry, const double InCurrentTime, const float InDeltaTime)
+{
+	TSharedPtr<FCurveEditor> CurveEditor = WeakCurveEditor.Pin();
+	if (!CurveEditor)
+	{
+		return;
+	}
 
-		OutputMax = FMath::Max(OutputMin + CurveInfoByID.Num() + ValueSpacePadding*(CurveInfoByID.Num()+1), 1.0);
+	if (!CurveEditor->AreBoundTransformUpdatesSuppressed())
+	{
+		UpdateViewToTransformCurves();
 	}
 
 	SInteractiveCurveEditorView::Tick(AllottedGeometry, InCurrentTime, InDeltaTime);
diff --git a/Engine/Source/Editor/CurveEditor/Public/CurveEditor.h b/Engine/Source/Editor/CurveEditor/Public/CurveEditor.h
--- a/Engine/Source/Editor/CurveEditor/Public/CurveEditor.h
+++ b/Engine/Source/Editor/CurveEditor/Public/CurveEditor.h
@@ -528,12 +528,6 @@ public:
 	bool CanFlattenOrStraightenSelection() const;
 
 public:
-	/**
-	 * Populate the specified array with curve painting parameters
-	 *
-	 * @param OutDrawParams    An array to populate with curve painting parameters, one per visible curve
-	 */
-	void GetCurveDrawParams(TArray<FCurveDrawParams>& OutDrawParams) const;
 
 	/**
 	 * Called by SCurveEditorPanel to update the allocated geometry for this curve editor.
diff --git a/Engine/Source/Editor/CurveEditor/Public/SCurveEditorView.h b/Engine/Source/Editor/CurveEditor/Public/SCurveEditorView.h
--- a/Engine/Source/Editor/CurveEditor/Public/SCurveEditorView.h
+++ b/Engine/Source/Editor/CurveEditor/Public/SCurveEditorView.h
@@ -149,6 +149,11 @@ public:
 	/** This should be called every tick by an owning widget, to see if the cache is valid, which will then recreate it and invalidate widget*/
 	virtual void CheckCacheAndInvalidateIfNeeded();
 
+	/** Function to make sure to update the view to the transform curves, we need to do this before we cache*/
+	virtual void UpdateViewToTransformCurves() {};
+
+	/** Frame the view vertially by the input and output bounds, peformaing any custom clipping as needed */
+	virtual void FrameVertical(double InOutputMin, double InOutputMax);
 public:
 
 	/**
@@ -211,7 +216,7 @@ public:
 protected:
 
 	/** Gets info about the curves being drawn. Converts actual curves into an abstract series of lines/points/handles/etc. */
-	void GetCurveDrawParams(TArray<FCurveDrawParams>& OutDrawParams) const;
+	void GetCurveDrawParams(TArray<FCurveDrawParams>& OutDrawParams);
 
 	/** Get it for just one curve*/
 	void GetCurveDrawParam(TSharedPtr<FCurveEditor>& CurveEditor, const FCurveModelID& ModelID, FCurveModel* CurveModel,
diff --git a/Engine/Source/Editor/CurveEditor/Public/Views/SCurveEditorViewNormalized.h b/Engine/Source/Editor/CurveEditor/Public/Views/SCurveEditorViewNormalized.h
--- a/Engine/Source/Editor/CurveEditor/Public/Views/SCurveEditorViewNormalized.h
+++ b/Engine/Source/Editor/CurveEditor/Public/Views/SCurveEditorViewNormalized.h
@@ -28,13 +28,15 @@ public:
 	/** Tools should use vertical snapping since grid lines to snap to will usually be visible */
 	virtual bool IsValueSnapEnabled() const override { return true; }
 
+	virtual void UpdateViewToTransformCurves() override;
+
+	void FrameVertical(double InOutputMin, double InOutputMax) override;
+
 protected:
 
 	virtual void Tick(const FGeometry& AllottedGeometry, const double InCurrentTime, const float InDeltaTime) override;
 
 	virtual void PaintView(const FPaintArgs& Args, const FGeometry& AllottedGeometry, const FSlateRect& MyCullingRect, FSlateWindowElementList& OutDrawElements, int32 BaseLayerId, const FWidgetStyle& InWidgetStyle, bool bParentEnabled) const override;
 
-	virtual void GetGridLinesY(TSharedRef<const FCurveEditor> CurveEditor, TArray<float>& MajorGridLines, TArray<float>& MinorGridLines, TArray<FText>* MajorGridLabels) const override;
-
 	virtual void DrawBufferedCurves(const FGeometry& AllottedGeometry, const FSlateRect& MyCullingRect, FSlateWindowElementList& OutDrawElements, int32 BaseLayerId, ESlateDrawEffect DrawEffects) const;
 };
diff --git a/Engine/Source/Editor/CurveEditor/Public/Views/SCurveEditorViewStacked.h b/Engine/Source/Editor/CurveEditor/Public/Views/SCurveEditorViewStacked.h
--- a/Engine/Source/Editor/CurveEditor/Public/Views/SCurveEditorViewStacked.h
+++ b/Engine/Source/Editor/CurveEditor/Public/Views/SCurveEditorViewStacked.h
@@ -29,6 +29,8 @@ public:
 	/** Tools should use vertical snapping since grid lines to snap to will usually be visible */
 	virtual bool IsValueSnapEnabled() const override { return true; }
 
+	virtual void UpdateViewToTransformCurves() override;
+
 protected:
 
 	virtual void Tick(const FGeometry& AllottedGeometry, const double InCurrentTime, const float InDeltaTime) override;
