Curve Editor: FindViews

#rnx

diff --git a/Engine/Source/Editor/CurveEditor/Public/SCurveEditorPanel.h b/Engine/Source/Editor/CurveEditor/Public/SCurveEditorPanel.h
--- a/Engine/Source/Editor/CurveEditor/Public/SCurveEditorPanel.h
+++ b/Engine/Source/Editor/CurveEditor/Public/SCurveEditorPanel.h
@@ -132,9 +132,9 @@ class CURVEEDITOR_API SCurveEditorPanel : public SCompoundWidget
 	 * @param InCurveID The identifier of the curve to find views for
 	 * @return An iterator to all the views that this cuvrve is displayed within.
 	 */
-	TMultiMap<FCurveModelID, TSharedRef<SCurveEditorView>>::TConstKeyIterator FindViews(FCurveModelID InCurveID)
+	TMultiMap<FCurveModelID, TSharedRef<SCurveEditorView>>::TConstKeyIterator FindViews(TRetainedRef<FCurveModelID> InCurveID)
 	{
-		return CurveViews.CreateConstKeyIterator(InCurveID);
+		return CurveViews.CreateConstKeyIterator(InCurveID.Get());
 	}
 
 	/**
