Curve Editor: Add Option to show bars(constarint/spaces)

diff --git a/Engine/Plugins/Animation/ControlRig/Source/ControlRigEditor/Private/Sequencer/ControlRigSpaceChannelEditors.cpp b/Engine/Plugins/Animation/ControlRig/Source/ControlRigEditor/Private/Sequencer/ControlRigSpaceChannelEditors.cpp
--- a/Engine/Plugins/Animation/ControlRig/Source/ControlRigEditor/Private/Sequencer/ControlRigSpaceChannelEditors.cpp
+++ b/Engine/Plugins/Animation/ControlRig/Source/ControlRigEditor/Private/Sequencer/ControlRigSpaceChannelEditors.cpp
@@ -1494,9 +1494,12 @@ TUniquePtr<FCurveModel> CreateCurveEditorModel(const TMovieSceneChannelHandle<FM
 	{
 		if (Channel->GetNumKeys() > 0)
 		{
-			return MakeUnique<FControlRigSpaceChannelCurveModel>(ChannelHandle, OwningSection, InSequencer);
+			const UCurveEditorSettings* Settings = GetDefault<UCurveEditorSettings>();
+			if (Settings == nullptr || Settings->GetShowBars())
+			{
+				return MakeUnique<FControlRigSpaceChannelCurveModel>(ChannelHandle, OwningSection, InSequencer);
+			}
 		}
-
 	}
 	return nullptr;
 }
diff --git a/Engine/Source/Editor/CurveEditor/Private/CurveEditor.cpp b/Engine/Source/Editor/CurveEditor/Private/CurveEditor.cpp
--- a/Engine/Source/Editor/CurveEditor/Private/CurveEditor.cpp
+++ b/Engine/Source/Editor/CurveEditor/Private/CurveEditor.cpp
@@ -477,6 +477,12 @@ void FCurveEditor::BindCommands()
 		FIsActionChecked::CreateLambda( [CurveSettings]{ return CurveSettings->GetAutoFrameCurveEditor(); } )
 	);
 
+	CommandList->MapAction(FCurveEditorCommands::Get().ToggleShowBars,
+		FExecuteAction::CreateLambda([this, CurveSettings] { CurveSettings->SetShowBars(!CurveSettings->GetShowBars()); Tree.RecreateModelsFromExistingSelection(this); }),
+		FCanExecuteAction(),
+		FIsActionChecked::CreateLambda([CurveSettings] { return CurveSettings->GetShowBars(); })
+	);
+
 	CommandList->MapAction(FCurveEditorCommands::Get().ToggleSnapTimeToSelection,
 		FExecuteAction::CreateLambda( [CurveSettings]{ CurveSettings->SetSnapTimeToSelection( !CurveSettings->GetSnapTimeToSelection() ); } ),
 		FCanExecuteAction(),
diff --git a/Engine/Source/Editor/CurveEditor/Private/CurveEditorCommands.cpp b/Engine/Source/Editor/CurveEditor/Private/CurveEditorCommands.cpp
--- a/Engine/Source/Editor/CurveEditor/Private/CurveEditorCommands.cpp
+++ b/Engine/Source/Editor/CurveEditor/Private/CurveEditorCommands.cpp
@@ -59,6 +59,7 @@ void FCurveEditorCommands::RegisterCommands()
 	UI_COMMAND(ToggleSnapTimeToSelection, "Snap Time to Selection", "Snap the current time to the first selected key time.", EUserInterfaceActionType::ToggleButton, FInputChord() );
 	UI_COMMAND(ToggleShowBufferedCurves, "Buffered Curves", "Show buffered curves for the selected curves.", EUserInterfaceActionType::ToggleButton, FInputChord() );
 	UI_COMMAND(ToggleShowCurveEditorCurveToolTips, "Curve Tool Tips", "Show a tool tip with name and values when hovering over a curve.", EUserInterfaceActionType::ToggleButton, FInputChord() );
+	UI_COMMAND(ToggleShowBars, "Show Bars", "Show Bars like Constraints and Spaces", EUserInterfaceActionType::ToggleButton, FInputChord());
 
 	UI_COMMAND(AddKeyHovered, "Add Key", "Add a new key to this curve at the current position.", EUserInterfaceActionType::Button, FInputChord(EKeys::MiddleMouseButton) );
 	UI_COMMAND(PasteKeysHovered, "Paste", "Paste clipboard contents", EUserInterfaceActionType::Button, FInputChord(EModifierKey::Control, EKeys::V));
diff --git a/Engine/Source/Editor/CurveEditor/Private/CurveEditorSettings.cpp b/Engine/Source/Editor/CurveEditor/Private/CurveEditorSettings.cpp
--- a/Engine/Source/Editor/CurveEditor/Private/CurveEditorSettings.cpp
+++ b/Engine/Source/Editor/CurveEditor/Private/CurveEditorSettings.cpp
@@ -7,6 +7,7 @@
 UCurveEditorSettings::UCurveEditorSettings()
 {
 	bAutoFrameCurveEditor = true;
+	bShowBars = true;
 	FrameInputPadding = 50;
 	FrameOutputPadding = 50;
 	bShowBufferedCurves = true;
@@ -36,6 +37,20 @@ void UCurveEditorSettings::SetAutoFrameCurveEditor(bool InbAutoFrameCurveEditor)
 	}
 }
 
+bool UCurveEditorSettings::GetShowBars() const
+{
+	return bShowBars;
+}
+
+void UCurveEditorSettings::SetShowBars(bool InValue)
+{
+	if (bShowBars != InValue)
+	{
+		bShowBars = InValue;
+		SaveConfig();
+	}
+}
+
 int32 UCurveEditorSettings::GetFrameInputPadding() const
 {
 	return FrameInputPadding;
diff --git a/Engine/Source/Editor/CurveEditor/Private/SCurveEditorPanel.cpp b/Engine/Source/Editor/CurveEditor/Private/SCurveEditorPanel.cpp
--- a/Engine/Source/Editor/CurveEditor/Private/SCurveEditorPanel.cpp
+++ b/Engine/Source/Editor/CurveEditor/Private/SCurveEditorPanel.cpp
@@ -1029,6 +1029,7 @@ TSharedRef<SWidget> SCurveEditorPanel::MakeCurveEditorCurveViewOptionsMenu()
 	MenuBuilder.AddMenuEntry(FCurveEditorCommands::Get().ToggleAutoFrameCurveEditor);
 	MenuBuilder.AddMenuEntry(FCurveEditorCommands::Get().ToggleSnapTimeToSelection);
 	MenuBuilder.AddMenuEntry(FCurveEditorCommands::Get().ToggleShowBufferedCurves);
+	MenuBuilder.AddMenuEntry(FCurveEditorCommands::Get().ToggleShowBars);
 	MenuBuilder.AddMenuEntry(FCurveEditorCommands::Get().ToggleShowCurveEditorCurveToolTips);
 
 	MenuBuilder.BeginSection("Organize", LOCTEXT("CurveEditorMenuOrganizeHeader", "Organize"));
diff --git a/Engine/Source/Editor/CurveEditor/Private/Tree/CurveEditorTree.cpp b/Engine/Source/Editor/CurveEditor/Private/Tree/CurveEditorTree.cpp
--- a/Engine/Source/Editor/CurveEditor/Private/Tree/CurveEditorTree.cpp
+++ b/Engine/Source/Editor/CurveEditor/Private/Tree/CurveEditorTree.cpp
@@ -544,6 +544,19 @@ void FCurveEditorTree::RemoveFromSelection(TArrayView<const FCurveEditorTreeItem
 	}
 }
 
+ void  FCurveEditorTree::RecreateModelsFromExistingSelection(FCurveEditor* CurveEditor)
+{
+	// Ensure the new selection has valid curve models
+	for (TTuple<FCurveEditorTreeItemID, ECurveEditorTreeSelectionState> NewItem : Selection)
+	{
+		if (NewItem.Value != ECurveEditorTreeSelectionState::None)
+		{
+			GetItem(NewItem.Key).DestroyCurves(CurveEditor);
+			GetItem(NewItem.Key).GetOrCreateCurves(CurveEditor);
+		}
+	}
+}
+
 const TMap<FCurveEditorTreeItemID, ECurveEditorTreeSelectionState>& FCurveEditorTree::GetSelection() const
 {
 	return Selection;
diff --git a/Engine/Source/Editor/CurveEditor/Public/CurveEditorCommands.h b/Engine/Source/Editor/CurveEditor/Public/CurveEditorCommands.h
--- a/Engine/Source/Editor/CurveEditor/Public/CurveEditorCommands.h
+++ b/Engine/Source/Editor/CurveEditor/Public/CurveEditorCommands.h
@@ -74,6 +74,7 @@ public:
 	TSharedPtr<FUICommandInfo> ToggleSnapTimeToSelection;
 	TSharedPtr<FUICommandInfo> ToggleShowBufferedCurves;
 	TSharedPtr<FUICommandInfo> ToggleShowCurveEditorCurveToolTips;
+	TSharedPtr<FUICommandInfo> ToggleShowBars;
 
 	TSharedPtr<FUICommandInfo> SetRandomCurveColorsForSelected;
 	TSharedPtr<FUICommandInfo> SetCurveColorsForSelected;
diff --git a/Engine/Source/Editor/CurveEditor/Public/CurveEditorSettings.h b/Engine/Source/Editor/CurveEditor/Public/CurveEditorSettings.h
--- a/Engine/Source/Editor/CurveEditor/Public/CurveEditorSettings.h
+++ b/Engine/Source/Editor/CurveEditor/Public/CurveEditorSettings.h
@@ -89,6 +89,11 @@ public:
 	/** Sets whether or not the curve editor auto frames the selected curves. */
 	void SetAutoFrameCurveEditor(bool InbAutoFrameCurveEditor);
 
+	/** Gets whether or not the curve editor shows key bar style curves, like for constraints and spaces. */
+	bool GetShowBars() const;
+	/** Sets whether or not the curve editor shows key bar style curves, like for constraints and spaces. */
+	void SetShowBars(bool InShowBars);
+
 	/** Gets the number of pixels to pad input framing */
 	int32 GetFrameInputPadding() const;
 	/** Sets the number of pixels to pad input framing */
@@ -161,6 +166,9 @@ protected:
 	UPROPERTY( config, EditAnywhere, Category="Curve Editor" )
 	bool bAutoFrameCurveEditor;
 
+	UPROPERTY(config, EditAnywhere, Category = "Curve Editor")
+	bool bShowBars;
+
 	/* Number of pixels to add as padding in the input axis when framing curve keys */
 	UPROPERTY( config, EditAnywhere, Category="Curve Editor", meta=(ClampMin=0) )
 	int32 FrameInputPadding;
diff --git a/Engine/Source/Editor/CurveEditor/Public/Tree/CurveEditorTree.h b/Engine/Source/Editor/CurveEditor/Public/Tree/CurveEditorTree.h
--- a/Engine/Source/Editor/CurveEditor/Public/Tree/CurveEditorTree.h
+++ b/Engine/Source/Editor/CurveEditor/Public/Tree/CurveEditorTree.h
@@ -513,6 +513,11 @@ public:
 		return bIsDoingDirectSelection;
 	}
 
+	/**
+	Recreate the curve models from the existing selection, this may be needed in case of a setting change.
+	*/
+	void RecreateModelsFromExistingSelection(FCurveEditor* CurveEditor);
+
 private:
 
 	// Recursively removes children without removing them from the parent (assuming the parent is also being removed)
diff --git a/Engine/Source/Editor/MovieSceneTools/Private/Channels/ConstraintChannelEditor.cpp b/Engine/Source/Editor/MovieSceneTools/Private/Channels/ConstraintChannelEditor.cpp
--- a/Engine/Source/Editor/MovieSceneTools/Private/Channels/ConstraintChannelEditor.cpp
+++ b/Engine/Source/Editor/MovieSceneTools/Private/Channels/ConstraintChannelEditor.cpp
@@ -34,6 +34,7 @@
 #include "Styling/ISlateStyle.h"
 #include "Templates/Function.h"
 #include "TimeToPixel.h"
+#include "CurveEditorSettings.h"
 #include "UObject/NameTypes.h"
 #include "UObject/UnrealNames.h"
 #include "Widgets/SNullWidget.h"
@@ -367,7 +368,11 @@ TUniquePtr<FCurveModel> CreateCurveEditorModel(const TMovieSceneChannelHandle<FM
 	{
 		if (Channel->GetNumKeys() > 0)
 		{
-			return MakeUnique<FConstraintChannelCurveModel>(ChannelHandle, OwningSection, InSequencer);
+			const UCurveEditorSettings* Settings = GetDefault<UCurveEditorSettings>();
+			if (Settings == nullptr || Settings->GetShowBars())
+			{
+				return MakeUnique<FConstraintChannelCurveModel>(ChannelHandle, OwningSection, InSequencer);
+			}
 		}
 	}
 	return nullptr;
