Curve Editor: Allow updating curve proximities when there is a drag operation so that the curve tooltips show the proper time and value.

diff --git a/Engine/Source/Editor/CurveEditor/Private/Views/SInteractiveCurveEditorView.cpp b/Engine/Source/Editor/CurveEditor/Private/Views/SInteractiveCurveEditorView.cpp
--- a/Engine/Source/Editor/CurveEditor/Private/Views/SInteractiveCurveEditorView.cpp
+++ b/Engine/Source/Editor/CurveEditor/Private/Views/SInteractiveCurveEditorView.cpp
@@ -665,11 +665,6 @@ bool SInteractiveCurveEditorView::GetCurveWithinWidgetRange(const FSlateRect& Wi
 
 void SInteractiveCurveEditorView::UpdateCurveProximities(FVector2D MousePixel)
 {
-	if (DragOperation.IsSet())
-	{
-		return;
-	}
-
 	CurveProximities.Reset();
 	CachedToolTipData.Reset();
 
