Curve Editor: Always update the tooltips when the drag operation is set because the curve proximities aren't updated on every drag

diff --git a/Engine/Source/Editor/CurveEditor/Private/Views/SInteractiveCurveEditorView.cpp b/Engine/Source/Editor/CurveEditor/Private/Views/SInteractiveCurveEditorView.cpp
--- a/Engine/Source/Editor/CurveEditor/Private/Views/SInteractiveCurveEditorView.cpp
+++ b/Engine/Source/Editor/CurveEditor/Private/Views/SInteractiveCurveEditorView.cpp
@@ -716,7 +716,8 @@ void SInteractiveCurveEditorView::UpdateCurveProximities(FVector2D MousePixel)
 
 	Algo::SortBy(CurveProximities, [](TTuple<FCurveModelID, float> In) { return In.Get<1>(); });
 
-	if (CurveProximities.Num() > 0 && CurveProximities[0].Get<1>() < CurveViewConstants::HoverProximityThresholdPx)
+	// Also, set the cached tooltips if dragging because the curve proximity might not be updated during the drag
+	if (CurveProximities.Num() > 0 && (CurveProximities[0].Get<1>() < CurveViewConstants::HoverProximityThresholdPx || DragOperation.IsSet()))
 	{
 		const FCurveModel* HoveredCurve = CurveEditor->FindCurve(CurveProximities[0].Get<0>());
 		if (HoveredCurve)
