Curve Editor: Apply postundo on postredo

Fixes stray curve data remaining on redo

diff --git a/Engine/Source/Editor/CurveEditor/Private/CurveEditor.cpp b/Engine/Source/Editor/CurveEditor/Private/CurveEditor.cpp
--- a/Engine/Source/Editor/CurveEditor/Private/CurveEditor.cpp
+++ b/Engine/Source/Editor/CurveEditor/Private/CurveEditor.cpp
@@ -2348,6 +2348,11 @@ void FCurveEditor::PostUndo(bool bSuccess)
 	}
 }
 
+void FCurveEditor::PostRedo(bool bSuccess)
+{
+	PostUndo(bSuccess);
+}
+
 void FCurveEditor::OnCustomColorsChanged()
 {
 	for (TPair<FCurveModelID, TUniquePtr<FCurveModel>>& CurvePair : CurveData)
diff --git a/Engine/Source/Editor/CurveEditor/Public/CurveEditor.h b/Engine/Source/Editor/CurveEditor/Public/CurveEditor.h
--- a/Engine/Source/Editor/CurveEditor/Public/CurveEditor.h
+++ b/Engine/Source/Editor/CurveEditor/Public/CurveEditor.h
@@ -334,6 +334,7 @@ public:
 
 	// FEditorUndoClient
 	virtual void PostUndo(bool bSuccess) override;
+	virtual void PostRedo(bool bSuccess) override;
 	// ~FEditorUndoClient
 
 	const TArray<TSharedRef<ICurveEditorExtension>> GetEditorExtensions() const
