Curve Editor: Channel Curve Models in Sequencer Curve Editor now track interpolation type of neighboring keys when added.

diff --git a/Engine/Source/Editor/CurveEditor/Private/Views/SInteractiveCurveEditorView.cpp b/Engine/Source/Editor/CurveEditor/Private/Views/SInteractiveCurveEditorView.cpp
--- a/Engine/Source/Editor/CurveEditor/Private/Views/SInteractiveCurveEditorView.cpp
+++ b/Engine/Source/Editor/CurveEditor/Private/Views/SInteractiveCurveEditorView.cpp
@@ -1068,10 +1068,16 @@ FReply SInteractiveCurveEditorView::OnMouseButtonDown(const FGeometry& MyGeometr
 					FScopedTransaction Transaction(LOCTEXT("InsertKey", "Insert Key"));
 
 					FCurveEditorScreenSpace CurveSpace = GetCurveSpace(HoveredCurve.GetValue());
-					FKeyAttributes KeyAttributes = CurveEditor->GetDefaultKeyAttributes().Get();
 					double MouseTime = CurveSpace.ScreenToSeconds(MousePixel.X);
 					double MouseValue = CurveSpace.ScreenToValue(MousePixel.Y);
 
+					FKeyAttributes KeyAttributes = CurveEditor->GetDefaultKeyAttributes().Get();
+					if (KeyAttributes.HasInterpMode())
+					{
+						// Set interpolation mode based on surrounding keys, if any
+						KeyAttributes.SetInterpMode(CurveToAddTo->GetInterpolationMode(MouseTime, KeyAttributes.GetInterpMode()));
+					}
+
 					FCurveSnapMetrics SnapMetrics = CurveEditor->GetCurveSnapMetrics(HoveredCurve.GetValue());
 					MouseTime = SnapMetrics.SnapInputSeconds(MouseTime);
 					MouseValue = SnapMetrics.SnapOutput(MouseValue);
@@ -1510,8 +1516,6 @@ void SInteractiveCurveEditorView::AddKeyAtTime(const TSet<FCurveModelID>& ToCurv
 	FScopedTransaction Transaction(LOCTEXT("AddKeyAtTime", "Add Key"));
 	bool bAddedKey = false;
 
-	FKeyAttributes DefaultAttributes = CurveEditor->GetDefaultKeyAttribute().Get();
-
 	// Clear the selection set as we will be selecting all the new keys created.
 	CurveEditor->GetSelection().Clear();
 
@@ -1553,8 +1557,15 @@ void SInteractiveCurveEditorView::AddKeyAtTime(const TSet<FCurveModelID>& ToCurv
 			}
 			else
 			{
+				FKeyAttributes KeyAttributes = CurveEditor->GetDefaultKeyAttribute().Get();
+				if (KeyAttributes.HasInterpMode())
+				{
+					// Set interpolation mode based on surrounding keys, if any
+					KeyAttributes.SetInterpMode(CurveModel->GetInterpolationMode(EvalTime, KeyAttributes.GetInterpMode()));
+				}
+
 				// Add a key on this curve
-				NewKey = CurveModel->AddKey(FKeyPosition(EvalTime, CurveValue), DefaultAttributes);
+				NewKey = CurveModel->AddKey(FKeyPosition(EvalTime, CurveValue), KeyAttributes);
 			}
 
 			// Add the key to the selection set.
diff --git a/Engine/Source/Editor/CurveEditor/Public/CurveModel.h b/Engine/Source/Editor/CurveEditor/Public/CurveModel.h
--- a/Engine/Source/Editor/CurveEditor/Public/CurveModel.h
+++ b/Engine/Source/Editor/CurveEditor/Public/CurveModel.h
@@ -150,6 +150,15 @@ public:
 	 */
 	virtual void GetNeighboringKeys(const FKeyHandle InKeyHandle, TOptional<FKeyHandle>& OutPreviousKeyHandle, TOptional<FKeyHandle>& OutNextKeyHandle) const = 0;
 
+	/**
+	 * Get the interpolation mode to use at a specified time
+	 *
+	 * @param InTime						The time we are looking for an interpolation mode
+	 * @param DefaultInterpolationMode		Current default interpolation mode, returned if other keys not found or interpolation not supported
+	 * @return Interpolation mode to use at that frame
+	 */
+	virtual ERichCurveInterpMode GetInterpolationMode(const double& InTime, ERichCurveInterpMode DefaultInterpolationMode) const { return DefaultInterpolationMode; }
+
 	/**
 	 * Evaluate this curve at the specified time
 	 *
diff --git a/Engine/Source/Editor/MovieSceneTools/Private/Channels/ChannelCurveModel.cpp b/Engine/Source/Editor/MovieSceneTools/Private/Channels/ChannelCurveModel.cpp
--- a/Engine/Source/Editor/MovieSceneTools/Private/Channels/ChannelCurveModel.cpp
+++ b/Engine/Source/Editor/MovieSceneTools/Private/Channels/ChannelCurveModel.cpp
@@ -331,6 +331,40 @@ void FChannelCurveModel<ChannelType, ChannelValue, KeyType>::GetNeighboringKeys(
 	}
 }
 
+template <class ChannelType, class ChannelValue, class KeyType>
+ERichCurveInterpMode FChannelCurveModel<ChannelType, ChannelValue, KeyType>::GetInterpolationMode(const double& InTime, ERichCurveInterpMode DefaultInterpolationMode) const
+{
+	ChannelType* Channel = ChannelHandle.Get();
+	UMovieSceneSection* Section = WeakSection.Get();
+	if (Channel && Section)
+	{
+		FFrameRate TickResolution = Section->GetTypedOuter<UMovieScene>()->GetTickResolution();
+
+		TMovieSceneChannelData<ChannelValue> ChannelData = Channel->GetData();
+		TArrayView<const FFrameNumber> Times = ChannelData.GetTimes();
+		TArrayView<const ChannelValue> Values = ChannelData.GetValues();
+
+		const FFrameNumber InFrame = (InTime * TickResolution).RoundToFrame();
+
+		if (Times.Num() > 0)
+		{
+ 			int32 InterpolationIndex = Algo::LowerBound(Times, InFrame) - 1;
+ 			if (InterpolationIndex < 0)
+ 			{
+ 				InterpolationIndex = 0;
+ 			}
+			const FKeyHandle KeyHandle = ChannelData.GetHandle(InterpolationIndex);
+			TArrayView<const FKeyHandle> InKey(&KeyHandle, 1);
+			TArray<FKeyAttributes> KeyAttributes;
+			KeyAttributes.SetNum(1);
+			GetKeyAttributes(InKey, KeyAttributes);
+			return KeyAttributes[0].GetInterpMode();
+		}
+	}
+
+	return DefaultInterpolationMode;
+}
+
 template <class ChannelType, class ChannelValue, class KeyType>
 bool FChannelCurveModel<ChannelType, ChannelValue, KeyType>::Evaluate(double Time, double& OutValue) const
 {
diff --git a/Engine/Source/Editor/MovieSceneTools/Public/Channels/ChannelCurveModel.h b/Engine/Source/Editor/MovieSceneTools/Public/Channels/ChannelCurveModel.h
--- a/Engine/Source/Editor/MovieSceneTools/Public/Channels/ChannelCurveModel.h
+++ b/Engine/Source/Editor/MovieSceneTools/Public/Channels/ChannelCurveModel.h
@@ -58,6 +58,7 @@ public:
 	virtual void GetValueRange(double& MinValue, double& MaxValue) const override;
 	virtual int32 GetNumKeys() const override;
 	virtual void GetNeighboringKeys(const FKeyHandle InKeyHandle, TOptional<FKeyHandle>& OutPreviousKeyHandle, TOptional<FKeyHandle>& OutNextKeyHandle) const override;
+	virtual ERichCurveInterpMode GetInterpolationMode(const double& InTime, ERichCurveInterpMode DefaultInterpolationMode) const override;
 	virtual bool Evaluate(double ProspectiveTime, double& OutValue) const override;
 	virtual void AddKeys(TArrayView<const FKeyPosition> InKeyPositions, TArrayView<const FKeyAttributes> InAttributes, TArrayView<TOptional<FKeyHandle>>* OutKeyHandles) override;
 	virtual void RemoveKeys(TArrayView<const FKeyHandle> InKeys) override;
