Add ability for child class to override curve tool tip #rb andrew.rodham
#rnx

diff --git a/Engine/Source/Editor/CurveEditor/Private/Views/SInteractiveCurveEditorView.cpp b/Engine/Source/Editor/CurveEditor/Private/Views/SInteractiveCurveEditorView.cpp
--- a/Engine/Source/Editor/CurveEditor/Private/Views/SInteractiveCurveEditorView.cpp
+++ b/Engine/Source/Editor/CurveEditor/Private/Views/SInteractiveCurveEditorView.cpp
@@ -730,15 +730,30 @@ void SInteractiveCurveEditorView::UpdateCurveProximities(FVector2D MousePixel)
 			HoveredCurve->Evaluate(EvaluatedTime, EvaluatedValue);
 
 			FCachedToolTipData ToolTipData;
-			ToolTipData.Text = FText::Format(LOCTEXT("CurveEditorTooltipName", "Name: {0}"), HoveredCurve->GetLongDisplayName());
-			ToolTipData.EvaluatedTime = FText::Format(LOCTEXT("CurveEditorTime", "Time: {0}"), EvaluatedTime);
-			ToolTipData.EvaluatedValue = FText::Format(LOCTEXT("CurveEditorValue", "Value: {0}"), EvaluatedValue);
+			ToolTipData.Text = FormatToolTipCurveName(*HoveredCurve);
+			ToolTipData.EvaluatedTime = FormatToolTipTime(*HoveredCurve, EvaluatedTime);
+			ToolTipData.EvaluatedValue = FormatToolTipValue(*HoveredCurve, EvaluatedValue);
 			
 			CachedToolTipData = ToolTipData;
 		}
 	}
 }
 
+FText SInteractiveCurveEditorView::FormatToolTipCurveName(const FCurveModel& CurveModel) const
+{
+	return FText::Format(LOCTEXT("CurveEditorTooltipName", "Name: {0}"), CurveModel.GetLongDisplayName());
+}
+
+FText SInteractiveCurveEditorView::FormatToolTipTime(const FCurveModel& CurveModel, double EvaluatedTime) const
+{
+	return FText::Format(LOCTEXT("CurveEditorTime", "Time: {0}"), EvaluatedTime);
+}
+
+FText SInteractiveCurveEditorView::FormatToolTipValue(const FCurveModel& CurveModel, double EvaluatedValue) const
+{
+	return FText::Format(LOCTEXT("CurveEditorValue", "Value: {0}"), EvaluatedValue);
+}
+
 void SInteractiveCurveEditorView::OnMouseEnter(const FGeometry& MyGeometry, const FPointerEvent& MouseEvent)
 {
 	//CurveProximities.Reset();
diff --git a/Engine/Source/Editor/CurveEditor/Public/Views/SInteractiveCurveEditorView.h b/Engine/Source/Editor/CurveEditor/Public/Views/SInteractiveCurveEditorView.h
--- a/Engine/Source/Editor/CurveEditor/Public/Views/SInteractiveCurveEditorView.h
+++ b/Engine/Source/Editor/CurveEditor/Public/Views/SInteractiveCurveEditorView.h
@@ -112,6 +112,10 @@ protected:
 	virtual bool GetCurveWithinWidgetRange(const FSlateRect& WidgetRectangle, TArray<FCurvePointHandle>* OutPoints) const override;
 	virtual TOptional<FCurveModelID> GetHoveredCurve() const override;
 
+	virtual FText FormatToolTipCurveName(const FCurveModel& CurveModel) const;
+	virtual FText FormatToolTipTime(const FCurveModel& CurveModel, double EvaluatedTime) const;
+	virtual FText FormatToolTipValue(const FCurveModel& CurveModel, double EvaluatedValue) const;
+
 	virtual void PaintView(const FPaintArgs& Args, const FGeometry& AllottedGeometry, const FSlateRect& MyCullingRect, FSlateWindowElementList& OutDrawElements, int32 BaseLayerId, const FWidgetStyle& InWidgetStyle, bool bParentEnabled) const;
 
 protected:
