Curve Editor: Rework the curve editor filters view so that it is now accessible from the curve editor panel and will emit OnFilterClassChanged.

Added setting the Frame Number customization for the filter details view so that FFrameNumber can be in frame numbers/seconds for the bake interval

diff --git a/Engine/Source/Editor/CurveEditor/Private/Filters/SCurveEditorFilterPanel.cpp b/Engine/Source/Editor/CurveEditor/Private/Filters/SCurveEditorFilterPanel.cpp
--- a/Engine/Source/Editor/CurveEditor/Private/Filters/SCurveEditorFilterPanel.cpp
+++ b/Engine/Source/Editor/CurveEditor/Private/Filters/SCurveEditorFilterPanel.cpp
@@ -73,7 +73,6 @@ void SCurveEditorFilterPanel::Construct(const FArguments& InArgs, TSharedRef<FCu
 
 	FOnClassPicked OnPicked(FOnClassPicked::CreateSP(this, &SCurveEditorFilterPanel::SetFilterClass));
 
-
 	FPropertyEditorModule& PropertyEditor = FModuleManager::LoadModuleChecked<FPropertyEditorModule>("PropertyEditor");
 	FClassViewerModule& ClassViewerModule = FModuleManager::LoadModuleChecked<FClassViewerModule>("ClassViewer");
 	
@@ -88,7 +87,7 @@ void SCurveEditorFilterPanel::Construct(const FArguments& InArgs, TSharedRef<FCu
 	DetailsViewArgs.bShowOptions = false;
 
 	// Generate it and store a reference so we can update it with the right object later.
-	DetailView = PropertyEditor.CreateDetailView(DetailsViewArgs);
+	DetailsView = PropertyEditor.CreateDetailView(DetailsViewArgs);
 	TSharedRef<SWidget> ClassFilterWidget = ClassViewerModule.CreateClassViewer(Options, OnPicked);
 
 	ChildSlot
@@ -105,7 +104,7 @@ void SCurveEditorFilterPanel::Construct(const FArguments& InArgs, TSharedRef<FCu
 		+ SVerticalBox::Slot()
 		.FillHeight(1.f)
 		[
-			DetailView.ToSharedRef()
+			DetailsView.ToSharedRef()
 		]
 		
 		// Footer Row
@@ -143,20 +142,32 @@ void SCurveEditorFilterPanel::Construct(const FArguments& InArgs, TSharedRef<FCu
 void SCurveEditorFilterPanel::SetFilterClass(UClass* InClass)
 {
 	UObject* ClassCDO = InClass->GetDefaultObject<UObject>();
-	DetailView->SetObject(ClassCDO);
+
+	TSharedPtr<FCurveEditor> CurveEditor = WeakCurveEditor.Pin();
+	if (CurveEditor)
+	{
+		if (UCurveEditorFilterBase* Filter = Cast<UCurveEditorFilterBase>(ClassCDO))
+		{
+			Filter->InitializeFilter(CurveEditor.ToSharedRef());
+		}
+	}
+
+	DetailsView->SetObject(ClassCDO);
+
+	OnFilterClassChanged.ExecuteIfBound();
 }
 
 FReply SCurveEditorFilterPanel::OnApplyClicked()
 {
 	TSharedPtr<FCurveEditor> CurveEditor = WeakCurveEditor.Pin();
-	if (CurveEditor && DetailView->GetSelectedObjects().Num() > 0)
+	if (CurveEditor && DetailsView->GetSelectedObjects().Num() > 0)
 	{
 		const TMap<FCurveModelID, FKeyHandleSet>& SelectedKeys = CurveEditor->GetSelection().GetAll();
 		
 		FText TransactionText = FText::Format(NSLOCTEXT("CurveEditorFilterApply", "Filter Curves", "Filtered {0}|plural(one=Curve, other=Curves)"), CurveEditor->GetSelection().Count());
 		FScopedTransaction Transaction(TransactionText);
 
-		UCurveEditorFilterBase* Filter = Cast<UCurveEditorFilterBase>(DetailView->GetSelectedObjects()[0]);
+		UCurveEditorFilterBase* Filter = Cast<UCurveEditorFilterBase>(DetailsView->GetSelectedObjects()[0]);
 		TMap<FCurveModelID, FKeyHandleSet> OutKeysToSelect;
 		Filter->ApplyFilter(CurveEditor.ToSharedRef(), SelectedKeys, OutKeysToSelect);
 
@@ -174,22 +185,24 @@ FReply SCurveEditorFilterPanel::OnApplyClicked()
 
 bool SCurveEditorFilterPanel::CanApplyFilter() const
 {
-	bool bHasSelection = false;
-	bool bHasFilter = false;
+	if (DetailsView->GetSelectedObjects().Num() == 0)
+	{
+		return false;
+	}
+
+	UCurveEditorFilterBase* CurrentFilter = Cast<UCurveEditorFilterBase>(DetailsView->GetSelectedObjects()[0].Get());
 
 	TSharedPtr<FCurveEditor> CurveEditor = WeakCurveEditor.Pin();
-	if (CurveEditor.IsValid())
+	if (CurveEditor.IsValid() && CurrentFilter)
 	{
-		bHasSelection = CurveEditor->GetSelection().Count() > 0;
+		return CurrentFilter->CanApplyFilter(CurveEditor.ToSharedRef());
 	}
 
-	bHasFilter = DetailView->GetSelectedObjects().Num() > 0;
-
-	return bHasSelection && bHasFilter;
+	return false;
 }
 
 
-void SCurveEditorFilterPanel::OpenDialog(TSharedPtr<SWindow> RootWindow, TSharedRef<FCurveEditor> InHostCurveEditor, TSubclassOf<UCurveEditorFilterBase> DefaultFilterClass)
+TSharedPtr<SCurveEditorFilterPanel> SCurveEditorFilterPanel::OpenDialog(TSharedPtr<SWindow> RootWindow, TSharedRef<FCurveEditor> InHostCurveEditor, TSubclassOf<UCurveEditorFilterBase> DefaultFilterClass)
 {
 	TSharedPtr<SWindow> ExistingWindow = ExistingFilterWindow.Pin();
 	if (ExistingWindow.IsValid())
@@ -221,6 +234,8 @@ void SCurveEditorFilterPanel::OpenDialog(TSharedPtr<SWindow> RootWindow, TShared
 	);
 
 	ExistingFilterWindow = ExistingWindow;
+
+	return FilterPanel;
 }
 
 void SCurveEditorFilterPanel::CloseDialog()
@@ -236,9 +251,9 @@ void SCurveEditorFilterPanel::CloseDialog()
 FText SCurveEditorFilterPanel::GetCurrentFilterText() const
 {
 	FText CurrentFilterName = NSLOCTEXT("SCurveEditorFilterPanel", "NoFilterSelectedName", "None");
-	if (DetailView->GetSelectedObjects().Num() > 0)
+	if (DetailsView->GetSelectedObjects().Num() > 0)
 	{
-		UCurveEditorFilterBase* CurrentFilter = Cast<UCurveEditorFilterBase>(DetailView->GetSelectedObjects()[0].Get());
+		UCurveEditorFilterBase* CurrentFilter = Cast<UCurveEditorFilterBase>(DetailsView->GetSelectedObjects()[0].Get());
 		if (CurrentFilter)
 		{
 			CurrentFilterName = CurrentFilter->GetClass()->GetDisplayNameText();
diff --git a/Engine/Source/Editor/CurveEditor/Private/SCurveEditorPanel.cpp b/Engine/Source/Editor/CurveEditor/Private/SCurveEditorPanel.cpp
--- a/Engine/Source/Editor/CurveEditor/Private/SCurveEditorPanel.cpp
+++ b/Engine/Source/Editor/CurveEditor/Private/SCurveEditorPanel.cpp
@@ -1148,9 +1148,18 @@ void SCurveEditorPanel::ShowCurveFilterUI(TSubclassOf<UCurveEditorFilterBase> Fi
 	TSharedPtr<SDockTab> OwnerTab = TabManager.IsValid() ? TabManager->GetOwnerTab() : TSharedPtr<SDockTab>();
 	TSharedPtr<SWindow> RootWindow = OwnerTab.IsValid() ? OwnerTab->GetParentWindow() : TSharedPtr<SWindow>();
 
-	SCurveEditorFilterPanel::OpenDialog(RootWindow, CurveEditor.ToSharedRef(), FilterClass);
+	FilterPanel = SCurveEditorFilterPanel::OpenDialog(RootWindow, CurveEditor.ToSharedRef(), FilterClass);
+	FilterPanel->OnFilterClassChanged.BindRaw(this, &SCurveEditorPanel::FilterClassChanged);
+
+	FilterClassChanged();
 }
 
+void SCurveEditorPanel::FilterClassChanged()
+{
+	OnFilterClassChanged.ExecuteIfBound();
+}
+
+
 const FGeometry& SCurveEditorPanel::GetScrollPanelGeometry() const
 {
 	return ScrollBox->GetCachedGeometry();
diff --git a/Engine/Source/Editor/CurveEditor/Public/Filters/SCurveEditorFilterPanel.h b/Engine/Source/Editor/CurveEditor/Public/Filters/SCurveEditorFilterPanel.h
--- a/Engine/Source/Editor/CurveEditor/Public/Filters/SCurveEditorFilterPanel.h
+++ b/Engine/Source/Editor/CurveEditor/Public/Filters/SCurveEditorFilterPanel.h
@@ -21,7 +21,7 @@ class CURVEEDITOR_API SCurveEditorFilterPanel : public SCompoundWidget
 public:
 	SLATE_BEGIN_ARGS(SCurveEditorFilterPanel)
 	{}
-
+		
 	SLATE_END_ARGS()
 
 	void Construct(const FArguments& InArgs, TSharedRef<FCurveEditor> InCurveEditor, UClass* DefaultFilterClass);
@@ -29,10 +29,17 @@ public:
 
 public:
 	/** Call this to request opening a window containing this panel. */
-	static void OpenDialog(TSharedPtr<SWindow> RootWindow, TSharedRef<FCurveEditor> InHostCurveEditor, TSubclassOf<UCurveEditorFilterBase> DefaultFilterClass);
+	static TSharedPtr<SCurveEditorFilterPanel> OpenDialog(TSharedPtr<SWindow> RootWindow, TSharedRef<FCurveEditor> InHostCurveEditor, TSubclassOf<UCurveEditorFilterBase> DefaultFilterClass);
 	
 	/** Closes the dialog if there is one open. */
 	static void CloseDialog();
+
+	/** The details view for filter class properties */
+	TSharedPtr<class IDetailsView> GetDetailsView() const { return DetailsView; }
+
+	/** Delegate for when the chosen filter class has changed */
+	FSimpleDelegate OnFilterClassChanged;
+
 protected:
 	FReply OnApplyClicked();
 	bool CanApplyFilter() const;
@@ -42,7 +49,7 @@ private:
 	TWeakPtr<FCurveEditor> WeakCurveEditor;
 
 	/** The Details View in our UI that we update each time they choose a class. */
-	TSharedPtr<IDetailsView> DetailView;
+	TSharedPtr<IDetailsView> DetailsView;
 
 	/** Singleton for the pop-up window. */
 	static TWeakPtr<SWindow> ExistingFilterWindow;
diff --git a/Engine/Source/Editor/CurveEditor/Public/SCurveEditorPanel.h b/Engine/Source/Editor/CurveEditor/Public/SCurveEditorPanel.h
--- a/Engine/Source/Editor/CurveEditor/Public/SCurveEditorPanel.h
+++ b/Engine/Source/Editor/CurveEditor/Public/SCurveEditorPanel.h
@@ -37,6 +37,7 @@ class ITimeSliderController;
 class SCurveEditorToolProperties;
 class SCurveEditorView;
 class SCurveEditorViewContainer;
+class SCurveEditorFilterPanel;
 class SCurveKeyDetailPanel;
 class SScrollBox;
 class SWidget;
@@ -100,6 +101,14 @@ class CURVEEDITOR_API SCurveEditorPanel : public SCompoundWidget
 		return KeyDetailsView;
 	}
 
+	/**
+	 * Access the filter panel
+	 */
+	TSharedPtr<class SCurveEditorFilterPanel> GetFilterPanel() const
+	{
+		return FilterPanel;
+	}
+
 	void AddView(TSharedRef<SCurveEditorView> ViewToAdd);
 
 	void RemoveView(TSharedRef<SCurveEditorView> ViewToRemove);
@@ -150,6 +159,11 @@ class CURVEEDITOR_API SCurveEditorPanel : public SCompoundWidget
 
 	/** Reset Stored Min/Max's*/
 	void ResetMinMaxes();
+
+	/** Delegate for when the chosen filter class has changed */
+	FSimpleDelegate OnFilterClassChanged;
+	void FilterClassChanged();
+
 private:
 	// SWidget Interface
 	virtual void Tick(const FGeometry& AllottedGeometry, const double InCurrentTime, const float InDeltaTime) override;
@@ -313,6 +327,9 @@ private:
 	/** Edit panel */
 	TSharedPtr<SCurveKeyDetailPanel> KeyDetailsView;
 
+	/* Filter panel */
+	TSharedPtr<SCurveEditorFilterPanel> FilterPanel;
+
 	/** Tool options panel */
 	TSharedPtr<SCurveEditorToolProperties> ToolPropertiesPanel;
 
diff --git a/Engine/Source/Editor/Sequencer/Private/MVVM/CurveEditorExtension.cpp b/Engine/Source/Editor/Sequencer/Private/MVVM/CurveEditorExtension.cpp
--- a/Engine/Source/Editor/Sequencer/Private/MVVM/CurveEditorExtension.cpp
+++ b/Engine/Source/Editor/Sequencer/Private/MVVM/CurveEditorExtension.cpp
@@ -4,6 +4,7 @@
 #include "MVVM/Selection/Selection.h"
 
 #include "FrameNumberDetailsCustomization.h"
+#include "Filters/SCurveEditorFilterPanel.h"
 #include "Framework/Docking/TabManager.h"
 #include "IPropertyRowGenerator.h"
 #include "MVVM/ViewModels/SequencerEditorViewModel.h"
@@ -201,7 +202,7 @@ void FCurveEditorExtension::CreateCurveEditor(const FTimeSliderArgs& TimeSliderA
 			TimeSliderArgs, Sequencer, CurveEditorModel.ToSharedRef());
 
 	CurveEditorTreeView = SNew(SCurveEditorTree, CurveEditorModel);
-	TSharedRef<SCurveEditorPanel> CurveEditorWidget = SNew(SCurveEditorPanel, CurveEditorModel.ToSharedRef())
+	CurveEditorPanel = SNew(SCurveEditorPanel, CurveEditorModel.ToSharedRef())
 		// Grid lines match the color specified in FSequencerTimeSliderController::OnPaintViewArea
 		.GridLineTint(FLinearColor(0.f, 0.f, 0.f, 0.3f))
 		.ExternalTimeSliderController(CurveEditorTimeSliderController)
@@ -242,7 +243,7 @@ void FCurveEditorExtension::CreateCurveEditor(const FTimeSliderArgs& TimeSliderA
 
 	// Register an instanced custom property type layout to handle converting FFrameNumber from Tick Resolution to Display Rate.
 	TWeakPtr<ISequencer> WeakSequencer(Sequencer);
-	CurveEditorWidget->GetKeyDetailsView()->GetPropertyRowGenerator()->RegisterInstancedCustomPropertyTypeLayout(
+	CurveEditorPanel->GetKeyDetailsView()->GetPropertyRowGenerator()->RegisterInstancedCustomPropertyTypeLayout(
 			"FrameNumber", 
 			FOnGetPropertyTypeCustomizationInstance::CreateStatic(CreateFrameNumberCustomization, WeakSequencer));
 
@@ -252,14 +253,43 @@ void FCurveEditorExtension::CreateCurveEditor(const FTimeSliderArgs& TimeSliderA
 		FExecuteAction::CreateLambda([this] { FSlateApplication::Get().SetKeyboardFocus(CurveEditorSearchBox, EFocusCause::SetDirectly); })
 	);
 
-	CurveEditorPanel = SNew(SSequencerCurveEditor, CurveEditorWidget, Sequencer);
+	CurveEditorWidget = SNew(SSequencerCurveEditor, CurveEditorPanel.ToSharedRef(), Sequencer);
+
+	CurveEditorPanel->OnFilterClassChanged.BindRaw(this, &FCurveEditorExtension::FilterClassChanged);
 
 	// Check to see if the tab is already opened due to the saved window layout.
 	FTabId TabId = FTabId(FCurveEditorExtension::CurveEditorTabName);
 	TSharedPtr<SDockTab> ExistingCurveEditorTab = Sequencer->GetToolkitHost()->GetTabManager()->FindExistingLiveTab(TabId);
 	if (ExistingCurveEditorTab)
 	{
-		ExistingCurveEditorTab->SetContent(CurveEditorPanel.ToSharedRef());
+		ExistingCurveEditorTab->SetContent(CurveEditorWidget.ToSharedRef());
+	}
+}
+
+void FCurveEditorExtension::FilterClassChanged()
+{
+	TSharedPtr<FSequencerEditorViewModel> OwnerModel = WeakOwnerModel.Pin();
+	if (!ensure(OwnerModel))
+	{
+		return;
+	}
+
+	TSharedPtr<FSequencer> Sequencer = OwnerModel->GetSequencerImpl();
+	if (!ensure(Sequencer))
+	{
+		return;
+	}
+
+	if (CurveEditorPanel)
+	{
+		TSharedPtr<SCurveEditorFilterPanel> FilterPanel = CurveEditorPanel->GetFilterPanel();
+		if (FilterPanel)
+		{
+			TWeakPtr<ISequencer> WeakSequencer(Sequencer);
+			FilterPanel->GetDetailsView()->RegisterInstancedCustomPropertyTypeLayout(
+				"FrameNumber",
+				FOnGetPropertyTypeCustomizationInstance::CreateStatic(CreateFrameNumberCustomization, WeakSequencer));
+		}
 	}
 }
 
@@ -283,7 +313,7 @@ void FCurveEditorExtension::OpenCurveEditor()
 	TSharedPtr<SDockTab> CurveEditorTab = Sequencer->GetToolkitHost()->GetTabManager()->TryInvokeTab(TabId);
 	if (CurveEditorTab.IsValid())
 	{
-		CurveEditorTab->SetContent(CurveEditorPanel.ToSharedRef());
+		CurveEditorTab->SetContent(CurveEditorWidget.ToSharedRef());
 
 		const FSlateIcon SequencerGraphIcon = FSlateIcon(FAppStyle::GetAppStyleSetName(), "GenericCurveEditor.TabIcon");
 		CurveEditorTab->SetTabIcon(SequencerGraphIcon.GetIcon());
diff --git a/Engine/Source/Editor/Sequencer/Public/MVVM/CurveEditorExtension.h b/Engine/Source/Editor/Sequencer/Public/MVVM/CurveEditorExtension.h
--- a/Engine/Source/Editor/Sequencer/Public/MVVM/CurveEditorExtension.h
+++ b/Engine/Source/Editor/Sequencer/Public/MVVM/CurveEditorExtension.h
@@ -12,6 +12,7 @@
 
 class ISequencer;
 class IPropertyTypeCustomization;
+class SCurveEditorPanel;
 class SCurveEditorTree;
 struct FTimeSliderArgs;
 
@@ -62,6 +63,9 @@ private:
 	/** Syncs the selection of the curve editor with that of the sequencer */
 	void SyncSelection();
 
+	/** Executed when the filter class has changed */
+	void FilterClassChanged();
+
 private:
 
 	/** The sequencer editor we are extending with a curve editor */
@@ -73,8 +77,10 @@ private:
 	TSharedPtr<SCurveEditorTree> CurveEditorTreeView;
 	/** The search widget for filtering curves in the Curve Editor tree. */
 	TSharedPtr<SWidget> CurveEditorSearchBox;
+	/** The curve editor widget containing the curve editor panel */
+	TSharedPtr<SWidget> CurveEditorWidget;
 	/** The curve editor panel. This is created and updated even if it is not currently visible. */
-	TSharedPtr<SWidget> CurveEditorPanel;
+	TSharedPtr<SCurveEditorPanel> CurveEditorPanel;
 
 	friend class FCurveEditorIntegrationExtension;
 };
