Curve Editor: Rework filters so that they are not dependent upon selection, ie. you can bake a range of frames/times.

diff --git a/Engine/Plugins/Editor/CurveEditorTools/Source/CurveEditorTools/Private/Filters/CurveEditorFFTFilter.cpp b/Engine/Plugins/Editor/CurveEditorTools/Source/CurveEditorTools/Private/Filters/CurveEditorFFTFilter.cpp
--- a/Engine/Plugins/Editor/CurveEditorTools/Source/CurveEditorTools/Private/Filters/CurveEditorFFTFilter.cpp
+++ b/Engine/Plugins/Editor/CurveEditorTools/Source/CurveEditorTools/Private/Filters/CurveEditorFFTFilter.cpp
@@ -16,17 +16,20 @@ UCurveEditorFFTFilter::UCurveEditorFFTFilter()
 
 void UCurveEditorFFTFilter::ApplyFilter_Impl(TSharedRef<FCurveEditor> InCurveEditor, const TMap<FCurveModelID, FKeyHandleSet>& InKeysToOperateOn, TMap<FCurveModelID, FKeyHandleSet>& OutKeysToSelect)
 {
+	const FFrameRate DisplayRate = InCurveEditor->GetTimeSliderController() ? InCurveEditor->GetTimeSliderController()->GetDisplayRate() : FFrameRate();
+	const FFrameRate TickResolution = InCurveEditor->GetTimeSliderController() ? InCurveEditor->GetTimeSliderController()->GetTickResolution() : FFrameRate();
+
 	// FFT works on one curve at a time and expects the keys to have an even spacing between them.
 	// We re-use the Bake filter to bake the supplied curve, but before we do this we need to cache
 	// the original key times so that we can evaluate the filtered curve and update the original curve
 	// with the smoothed out values.
 	// Now that we've backed up the original time and values of the keys, we need to bake it down to have even spacing. This is important for FFT.
 	UCurveEditorBakeFilter* BakeFilter = UCurveEditorBakeFilter::StaticClass()->GetDefaultObject<UCurveEditorBakeFilter>();
-		
+	BakeFilter->InitializeFilter(InCurveEditor);
+	
 	// Since we're baking under the hood, we need to cache their bake interval, override it to a reasonable number, and restore it at the end.
-	float OriginalIntervalRate = BakeFilter->BakeIntervalInSeconds;
-	bool bOriginalUseFrameBake = BakeFilter->bUseFrameBake;
-	BakeFilter->bUseFrameBake = false;
+	FFrameNumber OriginalBakeInterval = BakeFilter->BakeInterval;
+	float OriginalBakeIntervalInSeconds = BakeFilter->BakeIntervalInSeconds;
 
 	TArray<FKeyHandle> OriginalKeyHandles;
 	TArray<FKeyPosition> SelectedKeyPositions;
@@ -61,7 +64,11 @@ void UCurveEditorFFTFilter::ApplyFilter_Impl(TSharedRef<FCurveEditor> InCurveEdi
 			continue;
 		}
 		//Set interval range divded by twice the number of keys.
-		BakeFilter->BakeIntervalInSeconds = (MaxKey - MinKey) /  ( 2 * OriginalKeyHandles.Num());
+
+		const double IntervalAsSeconds = (MaxKey - MinKey) / (2 * OriginalKeyHandles.Num());
+		BakeFilter->BakeIntervalInSeconds = IntervalAsSeconds;
+		BakeFilter->BakeInterval = FFrameRate::TransformTime(DisplayRate.AsFrameNumber(IntervalAsSeconds), DisplayRate, TickResolution).FrameNumber;
+		
 		// This stores the position of all of the original keys. Once we've filtered the curve we will need to sample it at these positions.
 		TArray<FKeyPosition> OriginalKeyPositions;
 		TArray<FKeyAttributes> OriginalKeyAttributes;
@@ -168,7 +175,7 @@ void UCurveEditorFFTFilter::ApplyFilter_Impl(TSharedRef<FCurveEditor> InCurveEdi
 	}
 
 	// Restore their Bake Filter settings.
-	BakeFilter->BakeIntervalInSeconds = OriginalIntervalRate;
-	BakeFilter->bUseFrameBake = bOriginalUseFrameBake;
+	BakeFilter->BakeInterval = OriginalBakeInterval;
+	BakeFilter->BakeIntervalInSeconds = OriginalBakeIntervalInSeconds;
 }
 
diff --git a/Engine/Source/Editor/CurveEditor/CurveEditor.Build.cs b/Engine/Source/Editor/CurveEditor/CurveEditor.Build.cs
--- a/Engine/Source/Editor/CurveEditor/CurveEditor.Build.cs
+++ b/Engine/Source/Editor/CurveEditor/CurveEditor.Build.cs
@@ -15,6 +15,7 @@ public class CurveEditor : ModuleRules
 				"EditorFramework",
 				"Engine",
                 "InputCore",
+				"PropertyEditor",
 				"Slate",
 				"SlateCore",
 				"TimeManagement",
diff --git a/Engine/Source/Editor/CurveEditor/Private/CurveEditorModule.cpp b/Engine/Source/Editor/CurveEditor/Private/CurveEditorModule.cpp
--- a/Engine/Source/Editor/CurveEditor/Private/CurveEditorModule.cpp
+++ b/Engine/Source/Editor/CurveEditor/Private/CurveEditorModule.cpp
@@ -7,10 +7,13 @@
 #include "CurveEditorTypes.h"
 #include "CurveEditorViewRegistry.h"
 #include "Delegates/Delegate.h"
+#include "Filters/CurveEditorBakeFilter.h"
+#include "Filters/CurveEditorBakeFilterCustomization.h"
 #include "HAL/PlatformCrt.h"
 #include "ICurveEditorModule.h"
 #include "Misc/CoreDelegates.h"
 #include "Modules/ModuleManager.h"
+#include "PropertyEditorModule.h"
 #include "ToolMenus.h"
 
 class FCurveEditorModule : public ICurveEditorModule
@@ -29,11 +32,15 @@ public:
 				FCoreDelegates::OnPostEngineInit.AddStatic(&FCurveEditorCommands::Register);
 			}
 		}
+
+		RegisterCustomizations();
 	}
 
 	virtual void ShutdownModule() override
 	{
 		FCurveEditorCommands::Unregister();
+
+		UnregisterCustomizations();
 	}
 
 	virtual FDelegateHandle RegisterEditorExtension(FOnCreateCurveEditorExtension InOnCreateCurveEditorExtension) override
@@ -49,7 +56,6 @@ public:
 		EditorExtensionDelegates.RemoveAll([=](const FOnCreateCurveEditorExtension& Delegate) { return Delegate.GetHandle() == InHandle; });
 	}
 
-
 	virtual FDelegateHandle RegisterToolExtension(FOnCreateCurveEditorToolExtension InOnCreateCurveEditorToolExtension) override
 	{
 		ToolExtensionDelegates.Add(InOnCreateCurveEditorToolExtension);
@@ -88,6 +94,20 @@ public:
 		return ToolExtensionDelegates;
 	}
 
+	void RegisterCustomizations()
+	{
+		FPropertyEditorModule& PropertyEditorModule = FModuleManager::Get().LoadModuleChecked<FPropertyEditorModule>("PropertyEditor");
+		PropertyEditorModule.RegisterCustomClassLayout(UCurveEditorBakeFilter::StaticClass()->GetFName(), FOnGetDetailCustomizationInstance::CreateStatic(&FCurveEditorBakeFilterCustomization::MakeInstance));
+		PropertyEditorModule.NotifyCustomizationModuleChanged();
+	}
+
+	void UnregisterCustomizations()
+	{
+		FPropertyEditorModule& PropertyEditorModule = FModuleManager::Get().LoadModuleChecked<FPropertyEditorModule>("PropertyEditor");
+		PropertyEditorModule.UnregisterCustomClassLayout(UCurveEditorBakeFilter::StaticClass()->GetFName());
+		PropertyEditorModule.NotifyCustomizationModuleChanged();
+	}
+
 private:
 	/** List of editor extension handler delegates Curve Editors will execute when they are created. */
 	TArray<FOnCreateCurveEditorExtension> EditorExtensionDelegates;
diff --git a/Engine/Source/Editor/CurveEditor/Private/Filters/CurveEditorBakeFilter.cpp b/Engine/Source/Editor/CurveEditor/Private/Filters/CurveEditorBakeFilter.cpp
--- a/Engine/Source/Editor/CurveEditor/Private/Filters/CurveEditorBakeFilter.cpp
+++ b/Engine/Source/Editor/CurveEditor/Private/Filters/CurveEditorBakeFilter.cpp
@@ -20,6 +20,41 @@
 #include "Templates/Tuple.h"
 #include "Templates/UnrealTemplate.h"
 
+
+void UCurveEditorBakeFilter::InitializeFilter_Impl(TSharedRef<FCurveEditor> InCurveEditor)
+{
+	if (!InCurveEditor->GetTimeSliderController())
+	{
+		return;
+	}
+
+	// There is a valid time slider controller, so we can use the display rate and tick resolution from it to define frames
+	bUseSeconds = false;
+
+	// Initialize Bake Interval and Custom Ranges if they haven't been initialized with these parameters
+	const FFrameRate DisplayRate = InCurveEditor->GetTimeSliderController()->GetDisplayRate();
+	const FFrameRate TickResolution = InCurveEditor->GetTimeSliderController()->GetTickResolution();
+
+	if (DisplayRate == InitialDisplayRate && TickResolution == InitialTickResolution)
+	{
+		return;
+	}
+
+	InitialDisplayRate = DisplayRate;
+	InitialTickResolution = TickResolution;
+
+	TRange<FFrameNumber> PlayRange = InCurveEditor->GetTimeSliderController()->GetPlayRange();
+
+	BakeInterval = FFrameRate::TransformTime(FFrameTime(1), DisplayRate, TickResolution).FrameNumber;
+	CustomRange.Min = PlayRange.GetLowerBoundValue();
+	CustomRange.Max = PlayRange.GetUpperBoundValue();
+}
+
+bool UCurveEditorBakeFilter::CanApplyFilter_Impl(TSharedRef<FCurveEditor> InCurveEditor)
+{
+	return bCustomRangeOverride || InCurveEditor->GetSelection().Count() > 0;
+}
+
 void UCurveEditorBakeFilter::ApplyFilter_Impl(TSharedRef<FCurveEditor> InCurveEditor, const TMap<FCurveModelID, FKeyHandleSet>& InKeysToOperateOn, TMap<FCurveModelID, FKeyHandleSet>& OutKeysToSelect)
 {
 	TArray<FKeyHandle> KeyHandles;
@@ -28,7 +63,24 @@ void UCurveEditorBakeFilter::ApplyFilter_Impl(TSharedRef<FCurveEditor> InCurveEd
 	TArray<FKeyPosition> NewKeyPositions;
 	TArray<FKeyAttributes> NewKeyAttributes;
 
-	for (const TTuple<FCurveModelID, FKeyHandleSet>& Pair : InKeysToOperateOn)
+	const FFrameRate DisplayRate = InCurveEditor->GetTimeSliderController() ? InCurveEditor->GetTimeSliderController()->GetDisplayRate() : FFrameRate();
+	const FFrameRate TickResolution = InCurveEditor->GetTimeSliderController() ? InCurveEditor->GetTimeSliderController()->GetTickResolution() : FFrameRate();
+
+	// @todo: This code is a bit convoluted because the filter API is set up to only ever operate on key selections,
+	//        which needs correcting.
+	TMap<FCurveModelID, FKeyHandleSet> LocalKeysToOperateOn;
+	const TMap<FCurveModelID, FKeyHandleSet>* KeysToOperateOn = &InKeysToOperateOn;
+
+	if (InKeysToOperateOn.Num() == 0 && bCustomRangeOverride)
+	{
+		KeysToOperateOn = &LocalKeysToOperateOn;
+		for (const TTuple<FCurveModelID, TUniquePtr<FCurveModel>>& CurvePair : InCurveEditor->GetCurves())
+		{
+			LocalKeysToOperateOn.Add(CurvePair.Key);
+		}
+	}
+
+	for (const TTuple<FCurveModelID, FKeyHandleSet>& Pair : *KeysToOperateOn)
 	{
 		FCurveModel* Curve = InCurveEditor->FindCurve(Pair.Key);
 
@@ -49,39 +101,64 @@ void UCurveEditorBakeFilter::ApplyFilter_Impl(TSharedRef<FCurveEditor> InCurveEd
 
 		// Find the hull of the range of the selected keys
 		double MinKey = TNumericLimits<double>::Max(), MaxKey = TNumericLimits<double>::Lowest();
-		for (FKeyPosition Key : SelectedKeyPositions)
+
+		if (bCustomRangeOverride)
+		{
+			if (bUseSeconds)
+			{
+				MinKey = CustomRangeMinInSeconds;
+				MaxKey = CustomRangeMaxInSeconds;
+			}
+			else
+			{
+				MinKey = DisplayRate.AsSeconds(FFrameRate::TransformTime(CustomRange.Min, TickResolution, DisplayRate));
+				MaxKey = DisplayRate.AsSeconds(FFrameRate::TransformTime(CustomRange.Max, TickResolution, DisplayRate));
+			}
+		}
+		else
 		{
-			MinKey = FMath::Min(Key.InputValue, MinKey);
-			MaxKey = FMath::Max(Key.InputValue, MaxKey);
+			for (FKeyPosition Key : SelectedKeyPositions)
+			{
+				MinKey = FMath::Min(Key.InputValue, MinKey);
+				MaxKey = FMath::Max(Key.InputValue, MaxKey);
+			}
 		}
 
 		// Get all keys that exist between the time range
 		KeyHandles.Reset();
 		Curve->GetKeys(*InCurveEditor, MinKey, MaxKey, TNumericLimits<double>::Lowest(), TNumericLimits<double>::Max(), KeyHandles);
 
-		if (KeyHandles.Num() > 1)
-		{
-			// Determine new times for new keys
-			double Interval = bUseFrameBake ? (double)(BakeIntervalInFrames.Value) * BakeRate.AsInterval() : BakeIntervalInSeconds;
+		// Determine new times for new keys
+		
+		double Interval = bUseSeconds ? BakeIntervalInSeconds : DisplayRate.AsSeconds(FFrameRate::TransformTime(BakeInterval, TickResolution, DisplayRate));
 
-			int32 NumKeysToAdd = FMath::FloorToInt((MaxKey - MinKey) / Interval) + 1;
+		int32 NumKeysToAdd = FMath::FloorToInt((MaxKey - MinKey) / Interval) + 1;
 
-			NewKeyPositions.Reset(NumKeysToAdd);
-			NewKeyAttributes.Reset(NumKeysToAdd);
+		NewKeyPositions.Reset(NumKeysToAdd);
+		NewKeyAttributes.Reset(NumKeysToAdd);
+
+		int32 KeyIndex = 1;
+		for (; KeyIndex < NumKeysToAdd-1; ++KeyIndex)
+		{
+			double NewTime = MinKey + KeyIndex * Interval;
 
-			for (int32 KeyIndex = 1; KeyIndex < NumKeysToAdd; ++KeyIndex)
+			FKeyPosition CurrentKey(NewTime, 0.0);
+			if (Curve->Evaluate(CurrentKey.InputValue, CurrentKey.OutputValue))
 			{
-				double NewTime = MinKey + KeyIndex * Interval;
+				NewKeyPositions.Add(CurrentKey);
+				NewKeyAttributes.Add(FKeyAttributes().SetInterpMode(RCIM_Linear));
+			}
+		}
 
-				// Don't create the last key if it has the same value as the last key. This avoids duplicates of the last
-				// key, but improves behavior when the interval doesn't align.
-				if (KeyIndex == NumKeysToAdd - 1)
-				{
-					if (FMath::IsNearlyEqual(NewTime, SelectedKeyPositions.Last().InputValue))
-					{
-						continue;
-					}
-				}
+		// Handle the last key
+		if (KeyIndex < NumKeysToAdd)
+		{
+			double NewTime = MinKey + KeyIndex * Interval;
+
+			// Don't create the last key if it has the same value as the last key. This avoids duplicates of the last
+			// key, but improves behavior when the interval doesn't align.
+			if (SelectedKeyPositions.Num() == 0 || !FMath::IsNearlyEqual(NewTime, SelectedKeyPositions.Last().InputValue))
+			{
 				FKeyPosition CurrentKey(NewTime, 0.0);
 				if (Curve->Evaluate(CurrentKey.InputValue, CurrentKey.OutputValue))
 				{
@@ -89,40 +166,52 @@ void UCurveEditorBakeFilter::ApplyFilter_Impl(TSharedRef<FCurveEditor> InCurveEd
 					NewKeyAttributes.Add(FKeyAttributes().SetInterpMode(RCIM_Linear));
 				}
 			}
+		}
 
-			Curve->Modify();
+		Curve->Modify();
 
-			// We want to store the newly added ones so we can add them to the users selection to mimic
-			// their selection pre-baking.
-			TArray<TOptional<FKeyHandle>> NewKeyHandles;
-			NewKeyHandles.SetNumUninitialized(NewKeyPositions.Num());
+		// We want to store the newly added ones so we can add them to the users selection to mimic
+		// their selection pre-baking.
+		TArray<TOptional<FKeyHandle>> NewKeyHandles;
+		NewKeyHandles.SetNumUninitialized(NewKeyPositions.Num());
 
-			TArrayView<TOptional<FKeyHandle>> NewKeyHandlesView = MakeArrayView(NewKeyHandles);
-			FKeyHandleSet& OutHandleSet = OutKeysToSelect.Add(Pair.Key);
+		TArrayView<TOptional<FKeyHandle>> NewKeyHandlesView = MakeArrayView(NewKeyHandles);
+		FKeyHandleSet& OutHandleSet = OutKeysToSelect.Add(Pair.Key);
 
-			// Manually add the first and last keys to the out set before we remove them below.
+		// Manually add the first and last keys to the out set before we remove them below.
+		if (KeyHandles.Num() > 0)
+		{
 			OutHandleSet.Add(KeyHandles[0], ECurvePointType::Key);
+		}
+		if (KeyHandles.Num() > 1)
+		{
 			OutHandleSet.Add(KeyHandles.Last(0), ECurvePointType::Key);
+		}
 
-			// We need to leave the first and the last key of the selection alone for two reasons;
-			// 1. Undo/Redo works better as selections aren't transacted so they don't handle keys being removed/re-added.
-			// by leaving the keys alone, they survive undo/redo operations.
-			// 2. With high interval settings, the shape of the curve can change drastically because the next interval would fall right
-			// outside the last key. If we remove the keys, then the shape changes in these cases. We avoid putting a new key on the last
-			// key (if it exists), to avoid any duplication. This preserves the shape of the curves better in most test cases.
+		// We need to leave the first and the last key of the selection alone for two reasons;
+		// 1. Undo/Redo works better as selections aren't transacted so they don't handle keys being removed/re-added.
+		// by leaving the keys alone, they survive undo/redo operations.
+		// 2. With high interval settings, the shape of the curve can change drastically because the next interval would fall right
+		// outside the last key. If we remove the keys, then the shape changes in these cases. We avoid putting a new key on the last
+		// key (if it exists), to avoid any duplication. This preserves the shape of the curves better in most test cases.
+		if (KeyHandles.Num() > 0)
+		{
 			KeyHandles.RemoveAt(0);
+		}
+		if (KeyHandles.Num() > 1)
+		{
 			KeyHandles.RemoveAt(KeyHandles.Num() - 1);
+		}
 
-			// Remove all the old in-between keys and add the new ones
-			Curve->RemoveKeys(KeyHandles);
-			Curve->AddKeys(NewKeyPositions, NewKeyAttributes, &NewKeyHandlesView);
+		// Remove all the old in-between keys and add the new ones
+		Curve->RemoveKeys(KeyHandles);
+		Curve->AddKeys(NewKeyPositions, NewKeyAttributes, &NewKeyHandlesView);
 
-			for (const TOptional<FKeyHandle>& Handle : NewKeyHandlesView)
+		for (const TOptional<FKeyHandle>& Handle : NewKeyHandlesView)
+		{
+			if (Handle.IsSet())
 			{
-				if (Handle.IsSet())
-				{
-					OutHandleSet.Add(Handle.GetValue(), ECurvePointType::Key);
-				}
+				OutHandleSet.Add(Handle.GetValue(), ECurvePointType::Key);
 			}
 		}
 	}
diff --git a/Engine/Source/Editor/CurveEditor/Private/Filters/CurveEditorBakeFilterCustomization.cpp b/Engine/Source/Editor/CurveEditor/Private/Filters/CurveEditorBakeFilterCustomization.cpp
new file mode 100644
--- /dev/null
+++ b/Engine/Source/Editor/CurveEditor/Private/Filters/CurveEditorBakeFilterCustomization.cpp
@@ -0,0 +1,47 @@
+// Copyright Epic Games, Inc. All Rights Reserved.
+
+#include "CurveEditorBakeFilterCustomization.h"
+
+#include "Filters/CurveEditorBakeFilter.h"
+#include "DetailLayoutBuilder.h"
+#include "DetailWidgetRow.h"
+
+void FCurveEditorBakeFilterCustomization::CustomizeDetails(IDetailLayoutBuilder& DetailBuilder)
+{
+	TArray<TWeakObjectPtr<UObject>> DetailObjects;
+	DetailBuilder.GetObjectsBeingCustomized(DetailObjects);
+	if (DetailObjects.Num() == 0)
+	{
+		return;
+	}
+
+	bool bUseSeconds = true;
+	for (TWeakObjectPtr<UObject> DetailObject : DetailObjects)
+	{
+		UCurveEditorBakeFilter* BakeFilter = CastChecked<UCurveEditorBakeFilter>(DetailObject.Get());
+		if (BakeFilter)
+		{
+			bUseSeconds = BakeFilter->bUseSeconds;
+		}
+	}
+
+	if (bUseSeconds)
+	{
+		TSharedPtr<IPropertyHandle> BakeIntervalPropertyHandle = DetailBuilder.GetProperty(GET_MEMBER_NAME_CHECKED(UCurveEditorBakeFilter, BakeInterval));
+		DetailBuilder.HideProperty(BakeIntervalPropertyHandle);
+
+		TSharedPtr<IPropertyHandle> CustomRangePropertyHandle = DetailBuilder.GetProperty(GET_MEMBER_NAME_CHECKED(UCurveEditorBakeFilter, CustomRange));
+		DetailBuilder.HideProperty(CustomRangePropertyHandle);
+	}
+	else
+	{
+		TSharedPtr<IPropertyHandle> BakeIntervalInSecondsPropertyHandle = DetailBuilder.GetProperty(GET_MEMBER_NAME_CHECKED(UCurveEditorBakeFilter, BakeIntervalInSeconds));
+		DetailBuilder.HideProperty(BakeIntervalInSecondsPropertyHandle);
+
+		TSharedPtr<IPropertyHandle> CustomRangeMinInSecondsPropertyHandle = DetailBuilder.GetProperty(GET_MEMBER_NAME_CHECKED(UCurveEditorBakeFilter, CustomRangeMinInSeconds));
+		DetailBuilder.HideProperty(CustomRangeMinInSecondsPropertyHandle);
+
+		TSharedPtr<IPropertyHandle> CustomRangeMaxInSecondsPropertyHandle = DetailBuilder.GetProperty(GET_MEMBER_NAME_CHECKED(UCurveEditorBakeFilter, CustomRangeMaxInSeconds));
+		DetailBuilder.HideProperty(CustomRangeMaxInSecondsPropertyHandle);
+	}
+}
diff --git a/Engine/Source/Editor/CurveEditor/Private/Filters/CurveEditorBakeFilterCustomization.h b/Engine/Source/Editor/CurveEditor/Private/Filters/CurveEditorBakeFilterCustomization.h
new file mode 100644
--- /dev/null
+++ b/Engine/Source/Editor/CurveEditor/Private/Filters/CurveEditorBakeFilterCustomization.h
@@ -0,0 +1,26 @@
+// Copyright Epic Games, Inc. All Rights Reserved.
+
+#pragma once
+
+#include "IDetailCustomization.h"
+
+class UCurveEditorBakeFilter;
+
+class IPropertyHandle;
+
+/**
+ * Customizes a UCurveEditorBakeFilter
+ */
+class FCurveEditorBakeFilterCustomization : public IDetailCustomization
+{
+public:
+
+	static TSharedRef<IDetailCustomization> MakeInstance()
+	{
+		return MakeShared<FCurveEditorBakeFilterCustomization>();
+	}
+
+	// IDetailCustomization interface
+	virtual void CustomizeDetails(class IDetailLayoutBuilder& DetailBuilder) override;
+	// End IDetailCustomization interface
+};
\ No newline at end of file
diff --git a/Engine/Source/Editor/CurveEditor/Private/Filters/CurveEditorFilterBase.cpp b/Engine/Source/Editor/CurveEditor/Private/Filters/CurveEditorFilterBase.cpp
new file mode 100644
--- /dev/null
+++ b/Engine/Source/Editor/CurveEditor/Private/Filters/CurveEditorFilterBase.cpp
@@ -0,0 +1,10 @@
+// Copyright Epic Games, Inc. All Rights Reserved.
+
+#include "Filters/CurveEditorFilterBase.h"
+#include "CurveEditor.h"
+
+
+bool UCurveEditorFilterBase::CanApplyFilter_Impl(TSharedRef<FCurveEditor> InCurveEditor)
+{
+	return InCurveEditor->GetSelection().Count() > 0;
+}
diff --git a/Engine/Source/Editor/CurveEditor/Public/Filters/CurveEditorBakeFilter.h b/Engine/Source/Editor/CurveEditor/Public/Filters/CurveEditorBakeFilter.h
--- a/Engine/Source/Editor/CurveEditor/Public/Filters/CurveEditorBakeFilter.h
+++ b/Engine/Source/Editor/CurveEditor/Public/Filters/CurveEditorBakeFilter.h
@@ -5,6 +5,7 @@
 #include "Containers/Map.h"
 #include "CurveEditorTypes.h"
 #include "Filters/CurveEditorFilterBase.h"
+#include "ICurveEditorToolExtension.h"
 #include "Misc/FrameNumber.h"
 #include "Templates/SharedPointer.h"
 #include "UObject/ObjectMacros.h"
@@ -17,6 +18,18 @@ class UObject;
 struct FCurveModelID;
 struct FKeyHandleSet;
 
+USTRUCT(BlueprintType)
+struct FCurveEditorBakeFilterRange
+{
+	GENERATED_BODY()
+
+	UPROPERTY(EditAnywhere, BlueprintReadWrite, Category = "Settings")
+	FFrameNumber Min = FFrameNumber(0);
+
+	UPROPERTY(EditAnywhere, BlueprintReadWrite, Category = "Settings")
+	FFrameNumber Max = FFrameNumber(10);
+};
+
 UCLASS(DisplayName="Bake")
 class CURVEEDITOR_API UCurveEditorBakeFilter : public UCurveEditorFilterBase
 {
@@ -24,21 +37,50 @@ class CURVEEDITOR_API UCurveEditorBakeFilter : public UCurveEditorFilterBase
 public:
 	UCurveEditorBakeFilter()
 	{
-		bUseFrameBake = true;
-		BakeIntervalInFrames = FFrameNumber(1);
-		BakeIntervalInSeconds = 0.1;
+		bUseSeconds = true;
+
+		BakeIntervalInSeconds = 0.1f;
+		CustomRangeMinInSeconds = 0.f;
+		CustomRangeMaxInSeconds = 1.f;
+
+		BakeInterval = FFrameNumber(1);
 	}
 protected:
+	virtual void InitializeFilter_Impl(TSharedRef<FCurveEditor> InCurveEditoro) override;
 	virtual void ApplyFilter_Impl(TSharedRef<FCurveEditor> InCurveEditor, const TMap<FCurveModelID, FKeyHandleSet>& InKeysToOperateOn, TMap<FCurveModelID, FKeyHandleSet>& OutKeysToSelect) override;
+	virtual bool CanApplyFilter_Impl(TSharedRef<FCurveEditor> InCurveEditor) override;
+
 public:
-	/** If true we will use frame interval to bake, else will use seconds interval  */
-	UPROPERTY(EditAnywhere, BlueprintReadWrite, Category="Settings")
-	bool bUseFrameBake;
-
-	/** The interval (in display rate frames) between baked keys. Only used if bUseFrameBake is true. */
-	UPROPERTY(EditAnywhere, BlueprintReadWrite, meta = (EditCondition = "bUseFrameBake"), Category = "Settings")
-	FFrameNumber BakeIntervalInFrames;
-	/** The interval (in seconds) between baked keys. Only used if bUseSnapRateForInterval is false. */
-	UPROPERTY(EditAnywhere, BlueprintReadWrite, meta = (EditCondition = "!bUseFrameBake"), Category = "Settings")
+
+	/** The interval between baked keys when there is no valid Display Rate and Tick Resolution. */
+	UPROPERTY(EditAnywhere, BlueprintReadWrite, Category = "Settings", meta = (ClampMin = "0.0", UIMin = "0.0"))
 	float BakeIntervalInSeconds;
+
+	/** The interval between baked keys. */
+	UPROPERTY(EditAnywhere, BlueprintReadWrite, Category = "Settings", meta = (ClampMin = "0", UIMin = "0"))
+	FFrameNumber BakeInterval;
+
+	/** When enabled, CustomRange will be used for the bake region. Otherwise the selected keys will be used. */
+	UPROPERTY(EditAnywhere, BlueprintReadWrite, Category = "Settings")
+	bool bCustomRangeOverride;
+
+	/** Specifies a custom range to use for baking when there is no valid Display Rate and Tick Resolution. */
+	UPROPERTY(EditAnywhere, BlueprintReadWrite, Category = "Settings", meta = (EditCondition = "bCustomRangeOverride"))
+	float CustomRangeMinInSeconds;
+
+	/** Specifies a custom range to use for baking when there is no valid Display Rate and Tick Resolution. */
+	UPROPERTY(EditAnywhere, BlueprintReadWrite, Category = "Settings", meta = (EditCondition = "bCustomRangeOverride"))
+	float CustomRangeMaxInSeconds;
+
+	/** Specifies a custom range to use for baking */
+	UPROPERTY(EditAnywhere, BlueprintReadWrite, Category = "Settings", meta = (EditCondition = "bCustomRangeOverride"))
+	FCurveEditorBakeFilterRange CustomRange;
+
+	/** When enabled, will use BakeIntervalInSeconds and CustomRangeInSeconds. Otherwise, use BakeInterval and CustomRange which are defined by the Display Rate and Tick Resolution. */
+	UPROPERTY()
+	bool bUseSeconds;
+
+private:
+	FFrameRate InitialDisplayRate;
+	FFrameRate InitialTickResolution;
 };
\ No newline at end of file
diff --git a/Engine/Source/Editor/CurveEditor/Public/Filters/CurveEditorFilterBase.h b/Engine/Source/Editor/CurveEditor/Public/Filters/CurveEditorFilterBase.h
--- a/Engine/Source/Editor/CurveEditor/Public/Filters/CurveEditorFilterBase.h
+++ b/Engine/Source/Editor/CurveEditor/Public/Filters/CurveEditorFilterBase.h
@@ -26,6 +26,23 @@ class CURVEEDITOR_API UCurveEditorFilterBase : public UObject
 	GENERATED_BODY()
 
 public:
+
+	/**
+	 * Initialize this filter and any parameters it depends on
+	 */
+	void InitializeFilter(TSharedRef<FCurveEditor> InCurveEditor)
+	{
+		InitializeFilter_Impl(InCurveEditor);
+	}
+
+	/**
+	 * Check whether this filter can be applied to the specified curve editor
+	 */
+	bool CanApplyFilter(TSharedRef<FCurveEditor> InCurveEditor)
+	{
+		return CanApplyFilter_Impl(InCurveEditor);
+	}
+
 	/** 
 	* Applies the filter to all keys on the specified curve.
 	* @param		InCurveEditor		The curve editor that owns the FCurveModelIDs to operate on.
@@ -120,7 +137,14 @@ public:
 	}
 
 protected:
+
+	/** Implementation function for initializing this filter */
+	virtual void InitializeFilter_Impl(TSharedRef<FCurveEditor> InCurveEditor) {}
+
 	/** An implementation must override this function to implement filtering functionality. This is named different and doesn't use function overloading due to C++ name resolution issues
 	* which prevent calling base class functions of the same name (even with different signatures) from a pointer to the derived class. */
 	virtual void ApplyFilter_Impl(TSharedRef<FCurveEditor> InCurveEditor, const TMap<FCurveModelID, FKeyHandleSet>& InKeysToOperateOn, TMap<FCurveModelID, FKeyHandleSet>& OutKeysToSelect) {}
+
+	/** Implementation function for determining whether this filter can be applied */
+	virtual bool CanApplyFilter_Impl(TSharedRef<FCurveEditor> InCurveEditor);
 };
