Curve Editor: Match CurrentTime tooltip with Playhead, ala ESequencerZoomPosition

#rnx

diff --git a/Engine/Source/Editor/CurveEditor/Public/CurveEditorSettings.h b/Engine/Source/Editor/CurveEditor/Public/CurveEditorSettings.h
--- a/Engine/Source/Editor/CurveEditor/Public/CurveEditorSettings.h
+++ b/Engine/Source/Editor/CurveEditor/Public/CurveEditorSettings.h
@@ -33,11 +33,11 @@ enum class ECurveEditorTangentVisibility : uint8
 UENUM()
 enum class ECurveEditorZoomPosition : uint8
 {
-	/** Current Time. */
-	CurrentTime,
+	/** Playhead. */
+	CurrentTime UMETA(DisplayName = "Playhead"),
 
 	/** Mouse Position. */
-	MousePosition,
+	MousePosition UMETA(DisplayName = "Mouse Position"),
 };
 
 /** Custom Color Object*/
