Curve Editor: Fix crash on shutdown unregistering customizations.

#rnx

diff --git a/Engine/Source/Editor/CurveEditor/Private/CurveEditorModule.cpp b/Engine/Source/Editor/CurveEditor/Private/CurveEditorModule.cpp
--- a/Engine/Source/Editor/CurveEditor/Private/CurveEditorModule.cpp
+++ b/Engine/Source/Editor/CurveEditor/Private/CurveEditorModule.cpp
@@ -103,9 +103,14 @@ public:
 
 	void UnregisterCustomizations()
 	{
-		FPropertyEditorModule& PropertyEditorModule = FModuleManager::Get().LoadModuleChecked<FPropertyEditorModule>("PropertyEditor");
-		PropertyEditorModule.UnregisterCustomClassLayout(UCurveEditorBakeFilter::StaticClass()->GetFName());
-		PropertyEditorModule.NotifyCustomizationModuleChanged();
+		if (UObjectInitialized() && !IsEngineExitRequested())
+		{
+			if (FPropertyEditorModule* PropertyEditorModule = FModuleManager::GetModulePtr<FPropertyEditorModule>("PropertyEditor"))
+			{
+				PropertyEditorModule->UnregisterCustomClassLayout(UCurveEditorBakeFilter::StaticClass()->GetFName());
+				PropertyEditorModule->NotifyCustomizationModuleChanged();
+			}
+		}
 	}
 
 private:
