Curve Editor: Fix if MinKey > MaxKey

#rnx

diff --git a/Engine/Source/Editor/CurveEditor/Private/Filters/CurveEditorBakeFilter.cpp b/Engine/Source/Editor/CurveEditor/Private/Filters/CurveEditorBakeFilter.cpp
--- a/Engine/Source/Editor/CurveEditor/Private/Filters/CurveEditorBakeFilter.cpp
+++ b/Engine/Source/Editor/CurveEditor/Private/Filters/CurveEditorBakeFilter.cpp
@@ -114,6 +114,13 @@ void UCurveEditorBakeFilter::ApplyFilter_Impl(TSharedRef<FCurveEditor> InCurveEd
 				MinKey = DisplayRate.AsSeconds(FFrameRate::TransformTime(CustomRange.Min, TickResolution, DisplayRate));
 				MaxKey = DisplayRate.AsSeconds(FFrameRate::TransformTime(CustomRange.Max, TickResolution, DisplayRate));
 			}
+
+			if (MaxKey < MinKey)
+			{
+				const double TempMinKey = MinKey;
+				MinKey = MaxKey;
+				MaxKey = TempMinKey;
+			}
 		}
 		else
 		{
