Fix Curve Editor Curve Name node updating

diff --git a/Engine/Source/Editor/CurveEditor/Private/Views/SCurveEditorViewAbsolute.cpp b/Engine/Source/Editor/CurveEditor/Private/Views/SCurveEditorViewAbsolute.cpp
--- a/Engine/Source/Editor/CurveEditor/Private/Views/SCurveEditorViewAbsolute.cpp
+++ b/Engine/Source/Editor/CurveEditor/Private/Views/SCurveEditorViewAbsolute.cpp
@@ -2,6 +2,7 @@
 
 #include "Views/SCurveEditorViewAbsolute.h"
 
+#include "CurveEditor.h"
 #include "Delegates/Delegate.h"
 #include "Fonts/SlateFontInfo.h"
 #include "Layout/Children.h"
@@ -25,9 +26,13 @@ void SCurveEditorViewAbsolute::Construct(const FArguments& InArgs, TWeakPtr<FCur
 	.VAlign(VAlign_Top)
 	.Padding(FMargin(0.f, CurveViewConstants::CurveLabelOffsetY, CurveViewConstants::CurveLabelOffsetX, 0.f))
 	[
-		SNew(STextBlock)
+		SAssignNew(CurveCaptionTextBlock, STextBlock)
 		.Font(FCoreStyle::Get().GetFontStyle("ToolTip.LargerFont"))
-		.ColorAndOpacity(this, &SCurveEditorViewAbsolute::GetCurveCaptionColor)
-		.Text(this, &SCurveEditorViewAbsolute::GetCurveCaption)
 	];
+
+	if (InCurveEditor.IsValid())
+	{
+		InCurveEditor.Pin()->GetTree()->Events.OnItemsChanged.AddSP(this, &SInteractiveCurveEditorView::UpdateCurveCaption);
+		InCurveEditor.Pin()->GetTree()->Events.OnSelectionChanged.AddSP(this, &SInteractiveCurveEditorView::UpdateCurveCaption);	
+	}
 }
\ No newline at end of file
diff --git a/Engine/Source/Editor/CurveEditor/Private/Views/SCurveEditorViewNormalized.cpp b/Engine/Source/Editor/CurveEditor/Private/Views/SCurveEditorViewNormalized.cpp
--- a/Engine/Source/Editor/CurveEditor/Private/Views/SCurveEditorViewNormalized.cpp
+++ b/Engine/Source/Editor/CurveEditor/Private/Views/SCurveEditorViewNormalized.cpp
@@ -51,11 +51,15 @@ void SCurveEditorViewNormalized::Construct(const FArguments& InArgs, TWeakPtr<FC
 	.VAlign(VAlign_Top)
 	.Padding(FMargin(0.f, CurveViewConstants::CurveLabelOffsetY, CurveViewConstants::CurveLabelOffsetX, 0.f))
 	[
-		SNew(STextBlock)
+		SAssignNew(CurveCaptionTextBlock, STextBlock)
 		.Font(FCoreStyle::Get().GetFontStyle("ToolTip.LargerFont"))
-		.ColorAndOpacity(this, &SCurveEditorViewNormalized::GetCurveCaptionColor)
-		.Text(this, &SCurveEditorViewNormalized::GetCurveCaption)
 	];
+
+	if (InCurveEditor.IsValid())
+	{
+		InCurveEditor.Pin()->GetTree()->Events.OnItemsChanged.AddSP(this, &SInteractiveCurveEditorView::UpdateCurveCaption);
+		InCurveEditor.Pin()->GetTree()->Events.OnSelectionChanged.AddSP(this, &SInteractiveCurveEditorView::UpdateCurveCaption);	
+	}
 }
 
 FTransform2D CalculateViewToCurveTransform(const double OutputMin, const double OutputMax, const double InCurveOutputMin, const double InCurveOutputMax)
diff --git a/Engine/Source/Editor/CurveEditor/Private/Views/SInteractiveCurveEditorView.cpp b/Engine/Source/Editor/CurveEditor/Private/Views/SInteractiveCurveEditorView.cpp
--- a/Engine/Source/Editor/CurveEditor/Private/Views/SInteractiveCurveEditorView.cpp
+++ b/Engine/Source/Editor/CurveEditor/Private/Views/SInteractiveCurveEditorView.cpp
@@ -150,38 +150,26 @@ void SInteractiveCurveEditorView::Construct(const FArguments& InArgs, TWeakPtr<F
 	SetToolTip(ToolTipWidget);
 }
 
-FText SInteractiveCurveEditorView::GetCurveCaption() const
+void SInteractiveCurveEditorView::UpdateCurveCaption() const
 {
-	TSharedPtr<FCurveEditor> CurveEditor = WeakCurveEditor.Pin();
-	if (CurveEditor && CurveInfoByID.Num() == 1)
+	if (CurveCaptionTextBlock.IsValid())
 	{
-		for (const TTuple<FCurveModelID, FCurveInfo>& Pair : CurveInfoByID)
+		const TSharedPtr<FCurveEditor> CurveEditor = WeakCurveEditor.Pin();
+		if (CurveEditor && CurveEditor->GetCurves().Num() == 1)
 		{
-			if (const FCurveModel* Curve = CurveEditor->FindCurve(Pair.Key))
+			for (const TPair<FCurveModelID, TUniquePtr<FCurveModel>>& Pair : CurveEditor->GetCurves())
 			{
-				return Curve->GetLongDisplayName();
-			}
-		}
-	}
-
-	return FText::GetEmpty();
-}
-
-FSlateColor SInteractiveCurveEditorView::GetCurveCaptionColor() const
-{
-	TSharedPtr<FCurveEditor> CurveEditor = WeakCurveEditor.Pin();
-	if (CurveEditor && CurveInfoByID.Num() == 1)
-	{
-		for (const TTuple<FCurveModelID, FCurveInfo>& Pair : CurveInfoByID)
-		{
-			if (const FCurveModel* Curve = CurveEditor->FindCurve(Pair.Key))
-			{
-				return Curve->GetColor();
+				if (Pair.Value.IsValid())
+				{
+					CurveCaptionTextBlock->SetText(Pair.Value->GetLongDisplayName());
+					CurveCaptionTextBlock->SetColorAndOpacity(Pair.Value->GetColor());
+					return;
+				}
 			}
 		}
 	}
-
-	return BackgroundTint.CopyWithNewOpacity(1.f);
+	CurveCaptionTextBlock->SetText(FText::GetEmpty());
+	CurveCaptionTextBlock->SetColorAndOpacity(BackgroundTint.CopyWithNewOpacity(1.f));
 }
 
 void SInteractiveCurveEditorView::GetGridLinesX(TSharedRef<const FCurveEditor> CurveEditor, TArray<float>& MajorGridLines, TArray<float>& MinorGridLines, TArray<FText>* MajorGridLabels) const
diff --git a/Engine/Source/Editor/CurveEditor/Public/Views/SInteractiveCurveEditorView.h b/Engine/Source/Editor/CurveEditor/Public/Views/SInteractiveCurveEditorView.h
--- a/Engine/Source/Editor/CurveEditor/Public/Views/SInteractiveCurveEditorView.h
+++ b/Engine/Source/Editor/CurveEditor/Public/Views/SInteractiveCurveEditorView.h
@@ -105,6 +105,8 @@ public:
 
 	virtual void BuildContextMenu(FMenuBuilder& MenuBuilder, TOptional<FCurvePointHandle> ClickedPoint, TOptional<FCurveModelID> HoveredCurveID);
 
+	void UpdateCurveCaption() const;
+
 protected:
 
 	// ~SCurveEditorView Interface
@@ -137,9 +139,6 @@ protected:
 	void DrawCurves(TSharedRef<FCurveEditor> CurveEditor, const FGeometry& AllottedGeometry, const FSlateRect& MyCullingRect, FSlateWindowElementList& OutDrawElements, int32 BaseLayerId, const FWidgetStyle& InWidgetStyle, ESlateDrawEffect DrawEffects) const;
 	void DrawBufferedCurves(TSharedRef<FCurveEditor> CurveEditor, const FGeometry& AllottedGeometry, const FSlateRect& MyCullingRect, FSlateWindowElementList& OutDrawElements, int32 BaseLayerId, const FWidgetStyle& InWidgetStyle, ESlateDrawEffect DrawEffects) const;
 
-	FSlateColor GetCurveCaptionColor() const;
-	FText GetCurveCaption() const;
-
 private:
 	void HandleDirectKeySelectionByMouse(TSharedPtr<FCurveEditor> CurveEditor, const FPointerEvent& MouseEvent, TOptional<FCurvePointHandle> MouseDownPoint);
 
@@ -193,6 +192,9 @@ protected:
 	/** Background tint for this widget */
 	FLinearColor BackgroundTint;
 
+	/** Curve Name to update manually due to the parent retainer widget */
+	TSharedPtr<STextBlock> CurveCaptionTextBlock;
+
 private:
 
 	/** (Optional) the current drag operation */
