Fix curve editor stacked view mode not displaying single channel curves

diff --git a/Engine/Source/Editor/CurveEditor/Private/SCurveEditorViewContainer.cpp b/Engine/Source/Editor/CurveEditor/Private/SCurveEditorViewContainer.cpp
--- a/Engine/Source/Editor/CurveEditor/Private/SCurveEditorViewContainer.cpp
+++ b/Engine/Source/Editor/CurveEditor/Private/SCurveEditorViewContainer.cpp
@@ -94,7 +94,11 @@ FVector2D SCurveEditorViewContainer::ComputeDesiredSize(float) const
 	}
 
 	const float PanelHeight = CurveEditor->GetPanel()->GetScrollPanelGeometry().GetLocalSize().Y - 1.f;
-	MyDesiredSize.Y += FMath::Max(MinimumPanelHeight, (PanelHeight - MyDesiredSize.Y) / NumStretchPanels) * NumStretchPanels;
+
+	if (NumStretchPanels > 0)
+	{
+		MyDesiredSize.Y += FMath::Max(MinimumPanelHeight, (PanelHeight - MyDesiredSize.Y) / NumStretchPanels) * NumStretchPanels;
+	}
 
 	MyDesiredSize.Y = FMath::Max(MyDesiredSize.Y, PanelHeight);
 	return MyDesiredSize;
