Sequencer: Change name from auto to smart auto
Ah forget to submit this
#trivial

diff --git a/Engine/Source/Editor/CurveEditor/Private/CurveEditorCommands.cpp b/Engine/Source/Editor/CurveEditor/Private/CurveEditorCommands.cpp
--- a/Engine/Source/Editor/CurveEditor/Private/CurveEditorCommands.cpp
+++ b/Engine/Source/Editor/CurveEditor/Private/CurveEditorCommands.cpp
@@ -24,7 +24,7 @@ void FCurveEditorCommands::RegisterCommands()
 
 	UI_COMMAND(InterpolationConstant, "Constant", "Constant interpolation", EUserInterfaceActionType::ToggleButton, FInputChord(EKeys::Five));
 	UI_COMMAND(InterpolationLinear, "Linear", "Linear interpolation", EUserInterfaceActionType::ToggleButton, FInputChord(EKeys::Four));
-	UI_COMMAND(InterpolationCubicSmartAuto, "Auto", "Cubic interpolation - Smart Automatic tangents", EUserInterfaceActionType::ToggleButton, FInputChord(EKeys::Zero));
+	UI_COMMAND(InterpolationCubicSmartAuto, "Smart Auto", "Cubic interpolation - Smart Automatic tangents", EUserInterfaceActionType::ToggleButton, FInputChord(EKeys::Zero));
 	UI_COMMAND(InterpolationCubicAuto, "Auto", "Cubic interpolation - Automatic tangents", EUserInterfaceActionType::ToggleButton, FInputChord(EKeys::One));
 	UI_COMMAND(InterpolationCubicUser, "User", "Cubic interpolation - User flat tangents", EUserInterfaceActionType::ToggleButton, FInputChord(EKeys::Two));
 	UI_COMMAND(InterpolationCubicBreak, "Break", "Cubic interpolation - User broken tangents", EUserInterfaceActionType::ToggleButton, FInputChord(EKeys::Three));
