Curve Editor: Refresh the retainer when the curve caption or the curve caption color changes

diff --git a/Engine/Source/Editor/CurveEditor/Private/Views/SInteractiveCurveEditorView.cpp b/Engine/Source/Editor/CurveEditor/Private/Views/SInteractiveCurveEditorView.cpp
--- a/Engine/Source/Editor/CurveEditor/Private/Views/SInteractiveCurveEditorView.cpp
+++ b/Engine/Source/Editor/CurveEditor/Private/Views/SInteractiveCurveEditorView.cpp
@@ -152,6 +152,8 @@ void SInteractiveCurveEditorView::Construct(const FArguments& InArgs, TWeakPtr<F
 
 FText SInteractiveCurveEditorView::GetCurveCaption() const
 {
+	FText CurveCaption;
+
 	TSharedPtr<FCurveEditor> CurveEditor = WeakCurveEditor.Pin();
 	if (CurveEditor && CurveInfoByID.Num() == 1)
 	{
@@ -159,16 +161,25 @@ FText SInteractiveCurveEditorView::GetCurveCaption() const
 		{
 			if (const FCurveModel* Curve = CurveEditor->FindCurve(Pair.Key))
 			{
-				return Curve->GetLongDisplayName();
+				CurveCaption = Curve->GetLongDisplayName();
+				break;
 			}
 		}
 	}
 
-	return FText::GetEmpty();
+	if (!CurveCaption.IdenticalTo(CachedCurveCaption))
+	{
+		CachedCurveCaption = CurveCaption;
+		const_cast<SInteractiveCurveEditorView*>(this)->RefreshRetainer();
+	}
+
+	return CurveCaption;
 }
 
 FSlateColor SInteractiveCurveEditorView::GetCurveCaptionColor() const
 {
+	FSlateColor CurveCaptionColor = BackgroundTint.CopyWithNewOpacity(1.f);
+
 	TSharedPtr<FCurveEditor> CurveEditor = WeakCurveEditor.Pin();
 	if (CurveEditor && CurveInfoByID.Num() == 1)
 	{
@@ -176,12 +187,19 @@ FSlateColor SInteractiveCurveEditorView::GetCurveCaptionColor() const
 		{
 			if (const FCurveModel* Curve = CurveEditor->FindCurve(Pair.Key))
 			{
-				return Curve->GetColor();
+				CurveCaptionColor = Curve->GetColor();
+				break;
 			}
 		}
 	}
 
-	return BackgroundTint.CopyWithNewOpacity(1.f);
+	if (CurveCaptionColor != CachedCurveCaptionColor)
+	{
+		CachedCurveCaptionColor = CurveCaptionColor;
+		const_cast<SInteractiveCurveEditorView*>(this)->RefreshRetainer();
+	}
+
+	return CurveCaptionColor;
 }
 
 void SInteractiveCurveEditorView::GetGridLinesX(TSharedRef<const FCurveEditor> CurveEditor, TArray<float>& MajorGridLines, TArray<float>& MinorGridLines, TArray<FText>* MajorGridLabels) const
diff --git a/Engine/Source/Editor/CurveEditor/Public/Views/SInteractiveCurveEditorView.h b/Engine/Source/Editor/CurveEditor/Public/Views/SInteractiveCurveEditorView.h
--- a/Engine/Source/Editor/CurveEditor/Public/Views/SInteractiveCurveEditorView.h
+++ b/Engine/Source/Editor/CurveEditor/Public/Views/SInteractiveCurveEditorView.h
@@ -217,4 +217,10 @@ private:
 
 	/** Cached location of the mouse relative to this widget each tick. This is so that command bindings related to the mouse cursor can create them at the right time. */
 	FVector2D CachedMousePosition;
+
+	/** Cached curve caption, used to determine when to refresh the retainer */
+	mutable FText CachedCurveCaption;
+
+	/** Cached curve caption color, used to determine when to refresh the retainer */
+	mutable FSlateColor CachedCurveCaptionColor;
 };
\ No newline at end of file
