Curve Editor - Undoing/Redo keys does not work.
#fix ensure per-curve drawing parameters are always reset - issue occured when deleting the last key on a curve, leaving the cached drawing params in place until getting user-focus or making another change to the curve

diff --git a/Engine/Source/Editor/CurveEditor/Private/SCurveEditorView.cpp b/Engine/Source/Editor/CurveEditor/Private/SCurveEditorView.cpp
--- a/Engine/Source/Editor/CurveEditor/Private/SCurveEditorView.cpp
+++ b/Engine/Source/Editor/CurveEditor/Private/SCurveEditorView.cpp
@@ -231,6 +231,9 @@ void SCurveEditorView::GetCurveDrawParam(TSharedPtr<FCurveEditor>& CurveEditor,c
 	TArray<FKeyHandle> VisibleKeys;
 	CurveModel->GetKeys(*CurveEditor, InputMin, InputMax, TNumericLimits<double>::Lowest(), TNumericLimits<double>::Max(), VisibleKeys);
 
+	// Always reset the points to cover case going from 1 to 0 keys
+	Params.Points.Reset(VisibleKeys.Num());
+	
 	if (VisibleKeys.Num())
 	{
 		ECurveEditorTangentVisibility TangentVisibility = CurveEditor->GetSettings()->GetTangentVisibility();
@@ -240,7 +243,6 @@ void SCurveEditorView::GetCurveDrawParam(TSharedPtr<FCurveEditor>& CurveEditor,c
 
 		AllKeyPositions.SetNum(VisibleKeys.Num());
 		AllKeyAttributes.SetNum(VisibleKeys.Num());
-		Params.Points.Reset(VisibleKeys.Num());
 
 		CurveModel->GetKeyPositions(VisibleKeys, AllKeyPositions);
 		CurveModel->GetKeyAttributes(VisibleKeys, AllKeyAttributes);
