Add cached expansion state values to the curve editor tree so expansion state is restored when selection changes instead of defaulting automatically to a closed state. There are just saved per session, not in any asset. Added an option path name string to ICurveEditorTreeItemExtension's, which is then used as a hash to store the expansion state, since the curve model id's aren't permanent between selections.

diff --git a/Engine/Source/Editor/CurveEditor/Private/Tree/CurveEditorTree.cpp b/Engine/Source/Editor/CurveEditor/Private/Tree/CurveEditorTree.cpp
--- a/Engine/Source/Editor/CurveEditor/Private/Tree/CurveEditorTree.cpp
+++ b/Engine/Source/Editor/CurveEditor/Private/Tree/CurveEditorTree.cpp
@@ -557,6 +557,44 @@ void FCurveEditorTree::RemoveFromSelection(TArrayView<const FCurveEditorTreeItem
 	}
 }
 
+ TArray<FCurveEditorTreeItemID> FCurveEditorTree::GetCachedExpandedItems() const
+ {
+	 TArray<FCurveEditorTreeItemID> ExpandedItems;
+	 for (const TPair<FCurveEditorTreeItemID, FCurveEditorTreeItem>& Item : Items)
+	 {
+		 const TOptional<FString> PathName = Item.Value.GetUniquePathName();
+		 if (PathName.IsSet())
+		 {
+			 int32 Hash = GetTypeHash(PathName.GetValue());
+			 if (CachedExpandedItems.Contains(Hash))
+			 {
+				 ExpandedItems.Add(Item.Key);
+			 }
+		 }
+	 }
+	 return ExpandedItems;
+ }
+
+void FCurveEditorTree::SetItemExpansion(FCurveEditorTreeItemID InTreeItemID, bool bInExpansion)
+{
+	if (FCurveEditorTreeItem* TreeItem = FindItem(InTreeItemID))
+	{
+		TOptional<FString> OptString = TreeItem->GetUniquePathName();
+		if (OptString.IsSet())
+		{
+			int32 Hash = GetTypeHash(OptString.GetValue());
+			if (bInExpansion)
+			{
+				CachedExpandedItems.Add(Hash);
+			}
+			else
+			{
+				CachedExpandedItems.Remove(Hash);
+			}
+		}
+	}
+}
+
 const TMap<FCurveEditorTreeItemID, ECurveEditorTreeSelectionState>& FCurveEditorTree::GetSelection() const
 {
 	return Selection;
diff --git a/Engine/Source/Editor/CurveEditor/Private/Tree/SCurveEditorTree.cpp b/Engine/Source/Editor/CurveEditor/Private/Tree/SCurveEditorTree.cpp
--- a/Engine/Source/Editor/CurveEditor/Private/Tree/SCurveEditorTree.cpp
+++ b/Engine/Source/Editor/CurveEditor/Private/Tree/SCurveEditorTree.cpp
@@ -119,6 +119,7 @@ void SCurveEditorTree::Construct(const FArguments& InArgs, TSharedPtr<FCurveEdit
 		.OnGetChildren(this, &SCurveEditorTree::GetTreeItemChildren)
 		.OnGenerateRow(this, &SCurveEditorTree::GenerateRow)
 		.OnSetExpansionRecursive(this, &SCurveEditorTree::SetItemExpansionRecursive)
+		.OnExpansionChanged(this, &SCurveEditorTree::OnExpansionChanged)
 		.OnMouseButtonDoubleClick(InArgs._OnMouseButtonDoubleClick)
 		.OnContextMenuOpening(InArgs._OnContextMenuOpening)
 		.OnTreeViewScrolled(InArgs._OnTreeViewScrolled)
@@ -143,6 +144,8 @@ void SCurveEditorTree::RefreshTree()
 
 	const FCurveEditorTree* CurveEditorTree = CurveEditor->GetTree();
 	const FCurveEditorFilterStates& FilterStates = CurveEditorTree->GetFilterStates();
+	bool bExpansionWasCleared = false;
+	const TArray<FCurveEditorTreeItemID>& CachedExpandedItems = CurveEditor->GetTree()->GetCachedExpandedItems();
 
 	// When changing to/from a filtered state, we save and restore expansion states
 	if (FilterStates.IsActive() && !bFilterWasActive)
@@ -176,6 +179,7 @@ void SCurveEditorTree::RefreshTree()
 		}
 
 		// Restore expansion states
+		bExpansionWasCleared = true;
 		ClearExpandedItems();
 		for (FCurveEditorTreeItemID ExpandedItem : PreFilterExpandedItems)
 		{
@@ -195,7 +199,6 @@ void SCurveEditorTree::RefreshTree()
 
 	RootItems.Shrink();
 	RequestTreeRefresh();
-
 	if (FilterStates.IsActive())
 	{
 		TArray<FCurveEditorTreeItemID> ExpandedItems;
@@ -212,6 +215,7 @@ void SCurveEditorTree::RefreshTree()
 
 		if (ExpandedItems.Num() > 0)
 		{
+			bExpansionWasCleared = true;
 			ClearExpandedItems();
 			for (const FCurveEditorTreeItemID& ItemID : ExpandedItems)
 			{
@@ -220,6 +224,16 @@ void SCurveEditorTree::RefreshTree()
 		}
 	}
 
+	// at end we reset any cached expansions we still have set
+	if (bExpansionWasCleared == false)
+	{
+		ClearExpandedItems();
+	}
+	for (const FCurveEditorTreeItemID& ItemID : CachedExpandedItems)
+	{
+		SetItemExpansion(ItemID, true);
+	}
+
 	bFilterWasActive = FilterStates.IsActive();
 }
 
@@ -322,6 +336,14 @@ void SCurveEditorTree::OnTreeSelectionChanged(FCurveEditorTreeItemID, ESelectInf
 	}
 }
 
+void SCurveEditorTree::OnExpansionChanged(FCurveEditorTreeItemID Model, bool bInExpansionState)
+{
+	if (Model.IsValid())
+	{
+		CurveEditor->GetTree()->SetItemExpansion(Model, bInExpansionState);
+	}
+}
+
 void SCurveEditorTree::SetItemExpansionRecursive(FCurveEditorTreeItemID Model, bool bInExpansionState)
 {
 	if (Model.IsValid())
diff --git a/Engine/Source/Editor/CurveEditor/Public/Tree/CurveEditorTree.h b/Engine/Source/Editor/CurveEditor/Public/Tree/CurveEditorTree.h
--- a/Engine/Source/Editor/CurveEditor/Public/Tree/CurveEditorTree.h
+++ b/Engine/Source/Editor/CurveEditor/Public/Tree/CurveEditorTree.h
@@ -150,6 +150,16 @@ struct CURVEEDITOR_API FCurveEditorTreeItem
 		return Children.ChildIDs;
 	}
 
+	/**
+	*Get optional unique path name for the tree editor item
+	*/
+	TOptional<FString> GetUniquePathName() const { return UniquePathName; } 
+
+	/**
+	*Set optional unique path name for the tree editor item
+	*/
+	void SetUniquePathName(const TOptional<FString>& InName) { UniquePathName = InName; }
+
 	/**
 	 * Access the user-specified implementation for this tree item
 	 * @return A strong pointer to the implementation or null if it has expired, or was never assigned
@@ -208,6 +218,8 @@ private:
 	FCurveEditorTreeItemID ThisID;
 	/** This parent's ID or FCurveEditorTreeItemID::Invalid() for root nodes */
 	FCurveEditorTreeItemID ParentID;
+	/** Optional Unique Path Name*/
+	TOptional<FString> UniquePathName;
 	/** A weak pointer to an externally held implementation. Mutually exclusive to StrongItemImpl. */
 	TWeakPtr<ICurveEditorTreeItem> WeakItemImpl;
 	/** A strong pointer to an implementation for this tree item. Mutually exclusive to WeakItemImpl. */
@@ -505,6 +517,16 @@ public:
 		return ToggleExpansionStateDelegate;
 	}
 
+	/*
+	* Get cached expanded items
+	*/
+	TArray<FCurveEditorTreeItemID> GetCachedExpandedItems() const;
+	
+	/*
+	* Set Item expansion state
+	*/
+	void SetItemExpansion(FCurveEditorTreeItemID InTreeItemID, bool bInExpansion);
+
 	/**
 	Whether or not we are are doign a direct selection, could be used to see why a curve model is being created or destroyed, by direct selection or by sequencer filtering?
 	*/
@@ -563,6 +585,9 @@ private:
 	/** Delegate for when toggle expansion state is invoked */
 	FOnCurveEditorToggleExpansionState ToggleExpansionStateDelegate;
 
+	/** Set of cached expanded items, based on GetTypedHash(FString)*/
+	TSet<int32> CachedExpandedItems;
+
 	/** Whether or not we are are doign a direct selection, could be used to see why a curve model is being created or destroyed, by direct selection or by sequencer filtering?*/
 	bool bIsDoingDirectSelection = false;
 };
\ No newline at end of file
diff --git a/Engine/Source/Editor/CurveEditor/Public/Tree/SCurveEditorTree.h b/Engine/Source/Editor/CurveEditor/Public/Tree/SCurveEditorTree.h
--- a/Engine/Source/Editor/CurveEditor/Public/Tree/SCurveEditorTree.h
+++ b/Engine/Source/Editor/CurveEditor/Public/Tree/SCurveEditorTree.h
@@ -59,6 +59,8 @@ private:
 
 	void SetItemExpansionRecursive(FCurveEditorTreeItemID Model, bool bInExpansionState);
 
+	void OnExpansionChanged(FCurveEditorTreeItemID Model, bool bInExpansionState);
+
 	void RefreshTree();
 
 	void RefreshTreeWidgetSelection();
diff --git a/Engine/Source/Editor/Sequencer/Private/MVVM/CurveEditorIntegrationExtension.cpp b/Engine/Source/Editor/Sequencer/Private/MVVM/CurveEditorIntegrationExtension.cpp
--- a/Engine/Source/Editor/Sequencer/Private/MVVM/CurveEditorIntegrationExtension.cpp
+++ b/Engine/Source/Editor/Sequencer/Private/MVVM/CurveEditorIntegrationExtension.cpp
@@ -174,7 +174,8 @@ FCurveEditorTreeItemID FCurveEditorIntegrationExtension::AddToCurveEditor(TViewM
 	FCurveEditorTreeItem* NewItem = InCurveEditor->AddTreeItem(ParentID);
 	TSharedPtr<ICurveEditorTreeItem> CurveEditorTreeItem = InViewModel->GetCurveEditorTreeItem();
 	NewItem->SetWeakItem(CurveEditorTreeItem);
-
+	TOptional<FString> UniquePathName =  InViewModel->GetUniquePathName();
+	NewItem->SetUniquePathName(UniquePathName);
 	// Register the new ID in our map and notify the view model
 	ViewModelToTreeItemIDMap.Add(InViewModel.AsModel(), NewItem->GetID());
 	InViewModel->OnAddedToCurveEditor(NewItem->GetID(), InCurveEditor);
diff --git a/Engine/Source/Editor/Sequencer/Private/MVVM/ViewModels/ChannelModel.cpp b/Engine/Source/Editor/Sequencer/Private/MVVM/ViewModels/ChannelModel.cpp
--- a/Engine/Source/Editor/Sequencer/Private/MVVM/ViewModels/ChannelModel.cpp
+++ b/Engine/Source/Editor/Sequencer/Private/MVVM/ViewModels/ChannelModel.cpp
@@ -459,6 +459,15 @@ void FChannelGroupModel::CreateCurveModels(TArray<TUniquePtr<FCurveModel>>& OutC
 	}
 }
 
+TOptional<FString> FChannelGroupModel::GetUniquePathName() const
+{
+	TStringBuilder<256> StringBuilder;
+	IOutlinerExtension::GetPathName(*this, StringBuilder);
+	FString PathName(StringBuilder.ToString());
+	TOptional<FString> FullPathName = PathName;
+	return FullPathName;
+}
+
 bool FChannelGroupModel::HasCurves() const
 {
 	for (const TSharedRef<IKeyArea>& KeyArea : GetAllKeyAreas())
@@ -861,6 +870,11 @@ bool FChannelGroupOutlinerModel::HasCurves() const
 	return FChannelGroupModel::HasCurves();
 }
 
+TOptional<FString> FChannelGroupOutlinerModel::GetUniquePathName() const
+{
+	return FChannelGroupModel::GetUniquePathName();
+}
+
 bool FChannelGroupOutlinerModel::CanDelete(FText* OutErrorMessage) const
 {
 	return true;
diff --git a/Engine/Source/Editor/Sequencer/Private/MVVM/ViewModels/OutlinerItemModel.cpp b/Engine/Source/Editor/Sequencer/Private/MVVM/ViewModels/OutlinerItemModel.cpp
--- a/Engine/Source/Editor/Sequencer/Private/MVVM/ViewModels/OutlinerItemModel.cpp
+++ b/Engine/Source/Editor/Sequencer/Private/MVVM/ViewModels/OutlinerItemModel.cpp
@@ -684,6 +684,15 @@ bool FOutlinerItemModelMixin::HasCurves() const
 	return false;
 }
 
+TOptional<FString> FOutlinerItemModelMixin::GetUniquePathName() const
+{
+	TStringBuilder<256> StringBuilder;
+	IOutlinerExtension::GetPathName(*AsViewModel(), StringBuilder);
+	FString PathName(StringBuilder.ToString());
+	TOptional<FString> FullPathName = PathName;
+	return FullPathName;
+}
+
 TSharedPtr<ICurveEditorTreeItem> FOutlinerItemModelMixin::GetCurveEditorTreeItem() const
 {
 	TSharedRef<FViewModel> ThisShared(const_cast<FViewModel*>(AsViewModel())->AsShared());
diff --git a/Engine/Source/Editor/Sequencer/Public/MVVM/Extensions/ICurveEditorTreeItemExtension.h b/Engine/Source/Editor/Sequencer/Public/MVVM/Extensions/ICurveEditorTreeItemExtension.h
--- a/Engine/Source/Editor/Sequencer/Public/MVVM/Extensions/ICurveEditorTreeItemExtension.h
+++ b/Engine/Source/Editor/Sequencer/Public/MVVM/Extensions/ICurveEditorTreeItemExtension.h
@@ -31,6 +31,9 @@ public:
 	/** Whether this view-model has any curves and needs to have a curve editor item created */
 	virtual bool HasCurves() const = 0;
 
+	/** Get optional unique path name */
+	virtual TOptional<FString> GetUniquePathName() const = 0;
+
 	/** Gets the curve editor ID for this view-model */
 	virtual FCurveEditorTreeItemID GetCurveEditorItemID() const = 0;
 
diff --git a/Engine/Source/Editor/Sequencer/Public/MVVM/ViewModels/ChannelModel.h b/Engine/Source/Editor/Sequencer/Public/MVVM/ViewModels/ChannelModel.h
--- a/Engine/Source/Editor/Sequencer/Public/MVVM/ViewModels/ChannelModel.h
+++ b/Engine/Source/Editor/Sequencer/Public/MVVM/ViewModels/ChannelModel.h
@@ -185,6 +185,7 @@ public:
 
 	void CreateCurveModels(TArray<TUniquePtr<FCurveModel>>& OutCurveModels);
 	bool HasCurves() const;
+	TOptional<FString> GetUniquePathName() const;
 
 	void BuildChannelOverrideMenu(FMenuBuilder& MenuBuilder);
 
@@ -238,6 +239,7 @@ public:
 	/*~ ICurveEditorTreeItemExtension */
 	bool HasCurves() const override;
 	void BuildContextMenu(FMenuBuilder& MenuBuilder) override;
+	TOptional<FString> GetUniquePathName() const override;
 
 private:
 
diff --git a/Engine/Source/Editor/Sequencer/Public/MVVM/ViewModels/OutlinerItemModel.h b/Engine/Source/Editor/Sequencer/Public/MVVM/ViewModels/OutlinerItemModel.h
--- a/Engine/Source/Editor/Sequencer/Public/MVVM/ViewModels/OutlinerItemModel.h
+++ b/Engine/Source/Editor/Sequencer/Public/MVVM/ViewModels/OutlinerItemModel.h
@@ -54,6 +54,7 @@ public:
 	/*~ ICurveEditorTreeItemExtension */
 	virtual bool HasCurves() const override;
 	virtual TSharedPtr<ICurveEditorTreeItem> GetCurveEditorTreeItem() const override;
+	virtual TOptional<FString> GetUniquePathName() const override;
 
 	/*~ ICurveEditorTreeItem */
 	virtual TSharedPtr<SWidget> GenerateCurveEditorTreeWidget(const FName& InColumnName, TWeakPtr<FCurveEditor> InCurveEditor, FCurveEditorTreeItemID InTreeItemID, const TSharedRef<ITableRow>& InTableRow) override;
