Curve Editor: Show first and last key's tangent handle on both sides.

Previous, the left tangent handle for the first key and the right tangent handle for the last key would be hidden.

This is useful for example, if you have pre and post infinity set to cycle, you can manipulate the last key's leave tangent so that it's smooth to the first key's arrive tangent.

diff --git a/Engine/Source/Editor/CurveEditor/Private/RichCurveEditorModel.cpp b/Engine/Source/Editor/CurveEditor/Private/RichCurveEditorModel.cpp
--- a/Engine/Source/Editor/CurveEditor/Private/RichCurveEditorModel.cpp
+++ b/Engine/Source/Editor/CurveEditor/Private/RichCurveEditorModel.cpp
@@ -373,9 +373,6 @@ void FRichCurveEditorModel::GetKeyAttributes(TArrayView<const FKeyHandle> InKeys
 				return;
 			}
 
-			const FRichCurveKey* FirstKey = &AllKeys[0];
-			const FRichCurveKey* LastKey  = &AllKeys.Last();
-
 			for (int32 Index = 0; Index < InKeys.Num(); ++Index)
 			{
 				if (RichCurve.IsKeyHandleValid(InKeys[Index]))
@@ -390,24 +387,15 @@ void FRichCurveEditorModel::GetKeyAttributes(TArrayView<const FKeyHandle> InKeys
 					const bool bGetArriveTangent = RichCurve.IsKeyHandleValid(PreviousKeyHandle) && RichCurve.GetKeyRef(PreviousKeyHandle).InterpMode == RCIM_Cubic;
 					if (bGetArriveTangent)
 					{
-						if (&ThisKey != FirstKey)
-						{
-							Attributes.SetArriveTangent(ThisKey.ArriveTangent);
-						}
+						Attributes.SetArriveTangent(ThisKey.ArriveTangent);
 					}
 
 					if (ThisKey.InterpMode != RCIM_Constant && ThisKey.InterpMode != RCIM_Linear)
 					{
 						Attributes.SetTangentMode(ThisKey.TangentMode);
-						if (&ThisKey != FirstKey)
-						{
-							Attributes.SetArriveTangent(ThisKey.ArriveTangent);
-						}
+						Attributes.SetArriveTangent(ThisKey.ArriveTangent);
+						Attributes.SetLeaveTangent(ThisKey.LeaveTangent);
 
-						if (&ThisKey != LastKey)
-						{
-							Attributes.SetLeaveTangent(ThisKey.LeaveTangent);
-						}
 						if (ThisKey.InterpMode == RCIM_Cubic)
 						{
 							Attributes.SetTangentWeightMode(ThisKey.TangentWeightMode);
diff --git a/Engine/Source/Editor/MovieSceneTools/Private/Channels/BezierChannelCurveModel.cpp b/Engine/Source/Editor/MovieSceneTools/Private/Channels/BezierChannelCurveModel.cpp
--- a/Engine/Source/Editor/MovieSceneTools/Private/Channels/BezierChannelCurveModel.cpp
+++ b/Engine/Source/Editor/MovieSceneTools/Private/Channels/BezierChannelCurveModel.cpp
@@ -222,24 +222,15 @@ void FBezierChannelCurveModel<ChannelType, ChannelValue, KeyType>::GetKeyAttribu
 				const bool bGetArriveTangent = Values.IsValidIndex(PreviousKeyIndex) && Values[PreviousKeyIndex].InterpMode == RCIM_Cubic;
 				if (bGetArriveTangent)
 				{
-					if (KeyIndex != 0)
-					{
-						Attributes.SetArriveTangent(KeyValue.Tangent.ArriveTangent / TimeInterval);
-					}
+					Attributes.SetArriveTangent(KeyValue.Tangent.ArriveTangent / TimeInterval);
 				}
 
 				if ((KeyValue.InterpMode != RCIM_Constant && KeyValue.InterpMode != RCIM_Linear))
 				{
 					Attributes.SetTangentMode(KeyValue.TangentMode);
-					if (KeyIndex != 0)
-					{
-						Attributes.SetArriveTangent(KeyValue.Tangent.ArriveTangent / TimeInterval);
-					}
+					Attributes.SetArriveTangent(KeyValue.Tangent.ArriveTangent / TimeInterval);
+					Attributes.SetLeaveTangent(KeyValue.Tangent.LeaveTangent / TimeInterval);
 
-					if (KeyIndex != Times.Num()-1)
-					{
-						Attributes.SetLeaveTangent(KeyValue.Tangent.LeaveTangent / TimeInterval);
-					}
 					if (KeyValue.InterpMode == RCIM_Cubic)
 					{
 						Attributes.SetTangentWeightMode(KeyValue.Tangent.TangentWeightMode);
