Curve Editor: Change PixelsPerInput to double to fix jittery movement when manipulating keys

diff --git a/Engine/Plugins/Editor/CurveEditorTools/Source/CurveEditorTools/Private/Tools/CurveEditorRetimeTool.cpp b/Engine/Plugins/Editor/CurveEditorTools/Source/CurveEditorTools/Private/Tools/CurveEditorRetimeTool.cpp
--- a/Engine/Plugins/Editor/CurveEditorTools/Source/CurveEditorTools/Private/Tools/CurveEditorRetimeTool.cpp
+++ b/Engine/Plugins/Editor/CurveEditorTools/Source/CurveEditorTools/Private/Tools/CurveEditorRetimeTool.cpp
@@ -19,18 +19,18 @@
 #define LOCTEXT_NAMESPACE "CurveEditorToolCommands"
 namespace CurveEditorRetimeTool
 {
-	constexpr float AnchorWidth = 5.f;
-	constexpr float CloseButtonWidth = 18.f;
-	constexpr float CloseButtonPadding = 2.f;
-	constexpr float HighlightOpacity = 0.75f;
+	constexpr double AnchorWidth = 5.0;
+	constexpr double CloseButtonWidth = 18.0;
+	constexpr double CloseButtonPadding = 2.0;
+	constexpr double HighlightOpacity = 0.75;
 }
 
 void FCurveEditorRetimeAnchor::GetGeometry(const FGeometry& InWidgetGeometry, TSharedRef<FCurveEditor> InCurveEditor, FGeometry& OutBarGeometry, FGeometry& OutCloseButtonGeometry) const
 {
 	FCurveEditorScreenSpaceH HorizontalTransform = InCurveEditor->GetPanelInputSpace();
 
-	float AnchorPosition = HorizontalTransform.SecondsToScreen(ValueInSeconds);
-	float BarHeight = InWidgetGeometry.GetLocalSize().Y - (CurveEditorRetimeTool::CloseButtonWidth + CurveEditorRetimeTool::CloseButtonPadding);
+	double AnchorPosition = HorizontalTransform.SecondsToScreen(ValueInSeconds);
+	double BarHeight = InWidgetGeometry.GetLocalSize().Y - (CurveEditorRetimeTool::CloseButtonWidth + CurveEditorRetimeTool::CloseButtonPadding);
 
 	// Vertical Bar
 	FVector2D BarSize = FVector2D(CurveEditorRetimeTool::AnchorWidth, BarHeight);
@@ -346,7 +346,7 @@ void FCurveEditorRetimeTool::DrawAnchor(const FCurveEditorRetimeAnchor& InAnchor
 
 
 		// The width is measured from the top left corners. We subtract the width of one anchor (and then offset by that much) to not overlap the anchors.
-		const float GradientWidth = (NextBarGeometry.GetLocalPositionAtCoordinates(FVector2D(0,0)) - BarGeometry.GetLocalPositionAtCoordinates(FVector2D(0,0))).X - CurveEditorRetimeTool::AnchorWidth;
+		const double GradientWidth = (NextBarGeometry.GetLocalPositionAtCoordinates(FVector2D(0,0)) - BarGeometry.GetLocalPositionAtCoordinates(FVector2D(0,0))).X - CurveEditorRetimeTool::AnchorWidth;
 		const FVector2D GradientSize = FVector2D(GradientWidth, NextBarGeometry.GetLocalSize().Y);
 		const FVector2D GradientOffset = FVector2D(BarGeometry.GetLocalPositionAtCoordinates(FVector2D::ZeroVector) + FVector2D(CurveEditorRetimeTool::AnchorWidth, 0));
 		
diff --git a/Engine/Plugins/Runtime/WaveTable/Source/WaveTableEditor/Private/WaveTableCurveEditorViewStacked.cpp b/Engine/Plugins/Runtime/WaveTable/Source/WaveTableEditor/Private/WaveTableCurveEditorViewStacked.cpp
--- a/Engine/Plugins/Runtime/WaveTable/Source/WaveTableEditor/Private/WaveTableCurveEditorViewStacked.cpp
+++ b/Engine/Plugins/Runtime/WaveTable/Source/WaveTableEditor/Private/WaveTableCurveEditorViewStacked.cpp
@@ -247,8 +247,8 @@ namespace WaveTable
 				const int32  CurveIndexFromBottom = CurveInfoByID.Num() - Pair.Value.CurveIndex - 1;
 				const double PaddingToBottomOfView = (CurveIndexFromBottom + 1) * ValueSpacePadding;
 
-				const float PixelBottom = ViewSpace.ValueToScreen(CurveIndexFromBottom + PaddingToBottomOfView);
-				const float PixelTop = ViewSpace.ValueToScreen(CurveIndexFromBottom + PaddingToBottomOfView + 1.0);
+				const double PixelBottom = ViewSpace.ValueToScreen(CurveIndexFromBottom + PaddingToBottomOfView);
+				const double PixelTop = ViewSpace.ValueToScreen(CurveIndexFromBottom + PaddingToBottomOfView + 1.0);
 
 				if (!FSlateRect::DoRectanglesIntersect(MyCullingRect, TransformRect(AllottedGeometry.GetAccumulatedLayoutTransform(), FSlateRect(0, PixelTop, LocalSize.X, PixelBottom))))
 				{
diff --git a/Engine/Source/Editor/CurveEditor/Private/CurveEditor.cpp b/Engine/Source/Editor/CurveEditor/Private/CurveEditor.cpp
--- a/Engine/Source/Editor/CurveEditor/Private/CurveEditor.cpp
+++ b/Engine/Source/Editor/CurveEditor/Private/CurveEditor.cpp
@@ -690,7 +690,7 @@ void FCurveEditor::ZoomToFitInternal(EAxisList::Type Axes, const TMap<FCurveMode
 			TSharedPtr<SCurveEditorPanel> Panel = WeakPanel.Pin();
 			TSharedPtr<SCurveEditorView> View = WeakView.Pin();
 
-			int32 PanelWidth = 0;
+			double PanelWidth = 0;
 			if (Panel.IsValid())
 			{
 				PanelWidth = WeakPanel.Pin()->GetViewContainerGeometry().GetLocalSize().X;
@@ -700,7 +700,7 @@ void FCurveEditor::ZoomToFitInternal(EAxisList::Type Axes, const TMap<FCurveMode
 				PanelWidth = View->GetViewSpace().GetPhysicalWidth();
 			}
 			
-			double InputPercentage = PanelWidth != 0 ? FMath::Min(Settings->GetFrameInputPadding() / (float)PanelWidth, 0.5) : 0.1; // Cannot pad more than half the width
+			double InputPercentage = PanelWidth != 0 ? FMath::Min(Settings->GetFrameInputPadding() / PanelWidth, 0.5) : 0.1; // Cannot pad more than half the width
 
 			const double MinInputZoom = InputSnapEnabledAttribute.Get() ? InputSnapRateAttribute.Get().AsInterval() : 0.00001;
 			const double InputPadding = FMath::Max((InputMax - InputMin) * InputPercentage, MinInputZoom);
@@ -732,7 +732,7 @@ void FCurveEditor::ZoomToFitInternal(EAxisList::Type Axes, const TMap<FCurveMode
 		{
 			TSharedPtr<SCurveEditorPanel> Panel = WeakPanel.Pin();
 
-			int32 PanelHeight = 0;
+			double PanelHeight = 0;
 			if (Panel.IsValid())
 			{
 				PanelHeight = WeakPanel.Pin()->GetViewContainerGeometry().GetLocalSize().Y;
@@ -742,7 +742,7 @@ void FCurveEditor::ZoomToFitInternal(EAxisList::Type Axes, const TMap<FCurveMode
 				PanelHeight = View->GetViewSpace().GetPhysicalHeight();
 			}
 
-			double OutputPercentage = PanelHeight != 0 ? FMath::Min(Settings->GetFrameOutputPadding() / (float)PanelHeight, 0.5) : 0.1; // Cannot pad more than half the height
+			double OutputPercentage = PanelHeight != 0 ? FMath::Min(Settings->GetFrameOutputPadding() / PanelHeight, 0.5) : 0.1; // Cannot pad more than half the height
 
 			constexpr double MinOutputZoom = 0.00001;
 			const double OutputPadding = FMath::Max((OutputMax - OutputMin) * OutputPercentage, MinOutputZoom);
diff --git a/Engine/Source/Editor/CurveEditor/Private/CurveEditorHelpers.cpp b/Engine/Source/Editor/CurveEditor/Private/CurveEditorHelpers.cpp
--- a/Engine/Source/Editor/CurveEditor/Private/CurveEditorHelpers.cpp
+++ b/Engine/Source/Editor/CurveEditor/Private/CurveEditorHelpers.cpp
@@ -33,8 +33,8 @@ FVector2D ComputeScreenSpaceTangentOffset(const FCurveEditorScreenSpace& CurveSp
 
 void TangentAndWeightFromOffset(const FCurveEditorScreenSpace& CurveSpace, const FVector2D& TangentOffset, float& OutTangent, float& OutWeight)
 {
-	float X = CurveSpace.ScreenToSeconds(TangentOffset.X) - CurveSpace.ScreenToSeconds(0);
-	float Y = CurveSpace.ScreenToValue(TangentOffset.Y) - CurveSpace.ScreenToValue(0);
+	double X = CurveSpace.ScreenToSeconds(TangentOffset.X) - CurveSpace.ScreenToSeconds(0);
+	double Y = CurveSpace.ScreenToValue(TangentOffset.Y) - CurveSpace.ScreenToValue(0);
 
 	OutTangent = Y / X;
 	OutWeight = FMath::Sqrt(X*X + Y*Y);
@@ -49,9 +49,9 @@ FVector2D GetVectorFromSlopeAndLength(float Slope, float Length)
 
 void ConstructYGridLines(const FCurveEditorScreenSpace& ViewSpace, uint8 InMinorDivisions, TArray<float>& OutMajorGridLines, TArray<float>& OutMinorGridLines, FText GridLineLabelFormatY, TArray<FText>* OutMajorGridLabels)
 {
-	const float GridPixelSpacing = ViewSpace.GetPhysicalHeight() / 5.f;
+	const double GridPixelSpacing = ViewSpace.GetPhysicalHeight() / 5.0;
 
-	const float Order = FMath::Pow(10.f, FMath::FloorToInt(FMath::LogX(10.f, GridPixelSpacing / ViewSpace.PixelsPerOutput())));
+	const double Order = FMath::Pow(10.0, FMath::FloorToInt(FMath::LogX(10.0, GridPixelSpacing / ViewSpace.PixelsPerOutput())));
 
 	static const int32 DesirableBases[]  = { 2, 5 };
 	static const int32 NumDesirableBases = UE_ARRAY_COUNT(DesirableBases);
diff --git a/Engine/Source/Editor/CurveEditor/Private/SCurveEditorView.cpp b/Engine/Source/Editor/CurveEditor/Private/SCurveEditorView.cpp
--- a/Engine/Source/Editor/CurveEditor/Private/SCurveEditorView.cpp
+++ b/Engine/Source/Editor/CurveEditor/Private/SCurveEditorView.cpp
@@ -194,7 +194,7 @@ void SCurveEditorView::GetCurveDrawParam(TSharedPtr<FCurveEditor>& CurveEditor,c
 {
 
 	FCurveEditorScreenSpace CurveSpace = GetCurveSpace(ModelID);
-	const float DisplayRatio = (CurveSpace.PixelsPerOutput() / CurveSpace.PixelsPerInput());
+	const double DisplayRatio = (CurveSpace.PixelsPerOutput() / CurveSpace.PixelsPerInput());
 
 	const FKeyHandleSet* SelectedKeys = CurveEditor->GetSelection().GetAll().Find(ModelID);
 
@@ -256,8 +256,8 @@ void SCurveEditorView::GetCurveDrawParam(TSharedPtr<FCurveEditor>& CurveEditor,c
 				(TangentVisibility == ECurveEditorTangentVisibility::SelectedKeys && SelectedKeys &&
 					(SelectedKeys->Contains(VisibleKeys[Index], ECurvePointType::Any)));
 
-			float TimeScreenPos = CurveSpace.SecondsToScreen(KeyPosition.InputValue + InputOffset);
-			float ValueScreenPos = CurveSpace.ValueToScreen(KeyPosition.OutputValue);
+			double TimeScreenPos = CurveSpace.SecondsToScreen(KeyPosition.InputValue + InputOffset);
+			double ValueScreenPos = CurveSpace.ValueToScreen(KeyPosition.OutputValue);
 
 			// Add this key
 			FCurvePointInfo Key(KeyHandle);
@@ -315,7 +315,6 @@ void SCurveEditorView::GetCurveDrawParam(TSharedPtr<FCurveEditor>& CurveEditor,c
 				{
 					float PixelLength = 60.0f;
 					LeaveTangentPoint.ScreenPosition = Key.ScreenPosition + CurveEditor::GetVectorFromSlopeAndLength(LeaveTangent * -DisplayRatio, PixelLength);
-
 				}
 
 				LeaveTangentPoint.LineDelta = Key.ScreenPosition - LeaveTangentPoint.ScreenPosition;
diff --git a/Engine/Source/Editor/CurveEditor/Private/Views/SCurveEditorViewStacked.cpp b/Engine/Source/Editor/CurveEditor/Private/Views/SCurveEditorViewStacked.cpp
--- a/Engine/Source/Editor/CurveEditor/Private/Views/SCurveEditorViewStacked.cpp
+++ b/Engine/Source/Editor/CurveEditor/Private/Views/SCurveEditorViewStacked.cpp
@@ -265,8 +265,8 @@ void SCurveEditorViewStacked::DrawLabels(const FGeometry& AllottedGeometry, cons
 		const int32  CurveIndexFromBottom = CurveInfoByID.Num() - It->Value.CurveIndex - 1;
 		const double PaddingToBottomOfView = (CurveIndexFromBottom + 1)*ValueSpacePadding;
 		
-		const float PixelBottom = ViewSpace.ValueToScreen(CurveIndexFromBottom + PaddingToBottomOfView);
-		const float PixelTop = ViewSpace.ValueToScreen(CurveIndexFromBottom + PaddingToBottomOfView + 1.0);
+		const double PixelBottom = ViewSpace.ValueToScreen(CurveIndexFromBottom + PaddingToBottomOfView);
+		const double PixelTop = ViewSpace.ValueToScreen(CurveIndexFromBottom + PaddingToBottomOfView + 1.0);
 
 		if (!FSlateRect::DoRectanglesIntersect(MyCullingRect, TransformRect(AllottedGeometry.GetAccumulatedLayoutTransform(), FSlateRect(0, PixelTop, LocalSize.X, PixelBottom))))
 		{
diff --git a/Engine/Source/Editor/CurveEditor/Public/CurveEditorScreenSpace.h b/Engine/Source/Editor/CurveEditor/Public/CurveEditorScreenSpace.h
--- a/Engine/Source/Editor/CurveEditor/Public/CurveEditorScreenSpace.h
+++ b/Engine/Source/Editor/CurveEditor/Public/CurveEditorScreenSpace.h
@@ -19,8 +19,8 @@ struct FCurveEditorScreenSpaceH
 	/**
 	 * Construction from a physical size, and input/output range
 	 */
-	FCurveEditorScreenSpaceH(float InPixelWidth, double InInputMin, double InInputMax)
-		: PixelWidth(FMath::Max(InPixelWidth, 1.f))
+	FCurveEditorScreenSpaceH(double InPixelWidth, double InInputMin, double InInputMax)
+		: PixelWidth(FMath::Max(InPixelWidth, 1.0))
 		, InputMin(InInputMin), InputMax(InInputMax)
 	{
 		if (InputMax <= InputMin)
@@ -32,13 +32,13 @@ struct FCurveEditorScreenSpaceH
 public:
 
 	/** Convert a horizontal screen position in slate units to a value in seconds */
-	FORCEINLINE double ScreenToSeconds(float ScreenPosition) const
+	FORCEINLINE double ScreenToSeconds(double ScreenPosition) const
 	{
 		return InputMin + ScreenPosition / PixelsPerInput();
 	}
 
 	/** Convert a value in seconds to a horizontal screen position in slate units */
-	FORCEINLINE float SecondsToScreen(double InSeconds) const
+	FORCEINLINE double SecondsToScreen(double InSeconds) const
 	{
 		return (InSeconds - InputMin) * PixelsPerInput();
 	}
@@ -46,9 +46,9 @@ public:
 public:
 
 	/** Retrieve the number of slate units per input value */
-	FORCEINLINE float PixelsPerInput() const
+	FORCEINLINE double PixelsPerInput() const
 	{
-		float InputDiff = FMath::Max(InputMax - InputMin, 1e-10);
+		double InputDiff = FMath::Max(InputMax - InputMin, 1e-10);
 		return PixelWidth / InputDiff;
 	}
 
@@ -59,7 +59,7 @@ public:
 	/** Get the maximum input value displayed on the screen */
 	FORCEINLINE double GetInputMax() const { return InputMax; }
 	/** Get the physical width of the screen */
-	FORCEINLINE float GetPhysicalWidth() const { return PixelWidth; }
+	FORCEINLINE double GetPhysicalWidth() const { return PixelWidth; }
 
 public:
 
@@ -70,7 +70,7 @@ public:
 	{
 		const FVector2D& T = CurveTransform.GetTranslation();
 
-		float m00, m01, m10, m11;
+		double m00, m01, m10, m11;
 		CurveTransform.GetMatrix().GetMatrix(m00, m01, m10, m11);
 
 		double NewInputMin  = InputMin * m00 + T.X;
@@ -81,7 +81,7 @@ public:
 
 private:
 
-	float PixelWidth;
+	double PixelWidth;
 	double InputMin, InputMax;
 };
 
@@ -94,8 +94,8 @@ struct FCurveEditorScreenSpaceV
 	/**
 	 * Construction from a physical size, and input/output range
 	 */
-	FCurveEditorScreenSpaceV(float InPixelHeight, double InOutputMin, double InOutputMax)
-		: PixelHeight(FMath::Max(InPixelHeight, 1.f))
+	FCurveEditorScreenSpaceV(double InPixelHeight, double InOutputMin, double InOutputMax)
+		: PixelHeight(FMath::Max(InPixelHeight, 1.0))
 		, OutputMin(InOutputMin), OutputMax(InOutputMax)
 	{
 		if (!ensure(OutputMax > OutputMin))
@@ -107,13 +107,13 @@ struct FCurveEditorScreenSpaceV
 public:
 
 	/** Convert a vertical screen position in slate units to a value */
-	FORCEINLINE double ScreenToValue(float ScreenPosition) const
+	FORCEINLINE double ScreenToValue(double ScreenPosition) const
 	{
 		return OutputMin + (PixelHeight - ScreenPosition) / PixelsPerOutput();
 	}
 
 	/** Convert a value to a vertical screen position in slate units */
-	float ValueToScreen(double InValue) const
+	double ValueToScreen(double InValue) const
 	{
 		return (PixelHeight - (InValue - OutputMin) * PixelsPerOutput());
 	}
@@ -121,9 +121,9 @@ public:
 public:
 
 	/** Retrieve the number of slate units per output value */
-	FORCEINLINE float PixelsPerOutput() const
+	FORCEINLINE double PixelsPerOutput() const
 	{
-		float OutputDiff = FMath::Max(OutputMax - OutputMin, 1e-10);
+		double OutputDiff = FMath::Max(OutputMax - OutputMin, 1e-10);
 		return PixelHeight / OutputDiff;
 	}
 
@@ -134,7 +134,7 @@ public:
 	/** Get the maximum output value displayed on the screen */
 	FORCEINLINE double GetOutputMax() const { return OutputMax; }
 	/** Get the physical height of the screen */
-	FORCEINLINE float GetPhysicalHeight() const { return PixelHeight; }
+	FORCEINLINE double GetPhysicalHeight() const { return PixelHeight; }
 
 public:
 
@@ -145,7 +145,7 @@ public:
 	{
 		const FVector2D& T = CurveTransform.GetTranslation();
 
-		float m00, m01, m10, m11;
+		double m00, m01, m10, m11;
 		CurveTransform.GetMatrix().GetMatrix(m00, m01, m10, m11);
 
 		double NewOutputMin = OutputMin * m11 + T.Y;
@@ -156,7 +156,7 @@ public:
 
 private:
 
-	float PixelHeight;
+	double PixelHeight;
 	double OutputMin, OutputMax;
 };
 
