Curve Editor: Use FTextBlock so that orientation is handled automatically

Text was mirrored for Arabic

diff --git a/Engine/Source/Editor/CurveEditor/Private/SCurveEditorPanel.cpp b/Engine/Source/Editor/CurveEditor/Private/SCurveEditorPanel.cpp
--- a/Engine/Source/Editor/CurveEditor/Private/SCurveEditorPanel.cpp
+++ b/Engine/Source/Editor/CurveEditor/Private/SCurveEditorPanel.cpp
@@ -81,6 +81,7 @@
 #include "Widgets/SFrameRatePicker.h"
 #include "Widgets/SNullWidget.h"
 #include "Widgets/SOverlay.h"
+#include "Widgets/Text/STextBlock.h"
 
 class FPaintArgs;
 class FSlateRect;
@@ -105,45 +106,6 @@ static FAutoConsoleVariableRef CVarCurveEditorMaxCurvesPerPinnedView(
 	TEXT("When CurveEditor.PinnedViews is 1, defines the maximum number of curves allowed on a pinned view (0 for no maximum).")
 );
 
-/**
- * Utility class to hold an overlay which tells the user how to use the editor when there are no curves selected.
- */
-class SCurveEditorViewOverlay : public SCompoundWidget
-{
-	SLATE_BEGIN_ARGS(SCurveEditorViewOverlay)
-	{}
-	SLATE_END_ARGS()
-
-	void Construct(const FArguments& InArgs) {}
-
-
-	virtual int32 OnPaint(const FPaintArgs& Args, const FGeometry& AllottedGeometry, const FSlateRect& MyCullingRect, FSlateWindowElementList& OutDrawElements, int32 LayerId, const FWidgetStyle& InWidgetStyle, bool bParentEnabled) const override
-	{
-		static const FLinearColor BackgroundColor = FLinearColor::Black.CopyWithNewOpacity(0.35f);
-		static const FText InstructionText = LOCTEXT("CurveEditorTutorialOverlay", "Select a curve on the left to begin editing.");
-		const FSlateBrush*   WhiteBrush = FAppStyle::GetBrush("WhiteBrush");
-		const TSharedRef<FSlateFontMeasure> FontMeasure = FSlateApplication::Get().GetRenderer()->GetFontMeasureService();
-		const FSlateFontInfo FontInfo = FCoreStyle::Get().GetFontStyle("FontAwesome.13");
-
-		// Draw a darkened background
-		{
-			FSlateDrawElement::MakeBox(OutDrawElements, LayerId, AllottedGeometry.ToPaintGeometry(), WhiteBrush, ESlateDrawEffect::None, BackgroundColor);
-		}
-
-		// Draw the tutorial text centered
-		{
-			const FVector2D LabelSize = FontMeasure->Measure(InstructionText, FontInfo);
-			const FVector2D GeometrySize = AllottedGeometry.GetLocalSize();
-			const FVector2D LabelOffset = FVector2D((GeometrySize.X - LabelSize.X) / 2.f, (GeometrySize.Y - LabelSize.Y) / 2.f);
-			const FPaintGeometry LabelGeometry = AllottedGeometry.ToPaintGeometry(FSlateLayoutTransform(LabelOffset));
-
-			FSlateDrawElement::MakeText(OutDrawElements, LayerId, LabelGeometry, InstructionText, FontInfo);
-		}
-
-		return LayerId + 1;
-	}
-};
-
 /**
  * Implemented as a friend struct to SCurveEditorView to ensure that SCurveEditorPanel is the only thing that can add/remove curves from views
  * whilst disallowing access to any other private members;
@@ -282,8 +244,32 @@ void SCurveEditorPanel::Construct(const FArguments& InArgs, TSharedRef<FCurveEdi
 		// An overlay for the main area which lets us put system-wide overlays
 		+ SOverlay::Slot()
 		[
-			SNew(SCurveEditorViewOverlay)
+			SNew(SOverlay)
 			.Visibility(this, &SCurveEditorPanel::ShouldInstructionOverlayBeVisible)
+
+			// Darker background
+			+ SOverlay::Slot()
+			[
+				SNew(SBorder)
+				.BorderImage(FAppStyle::GetBrush("ToolPanel.GroupBorder"))
+				.BorderBackgroundColor(FLinearColor::Black.CopyWithNewOpacity(0.35f))
+			]
+
+			// Text
+			+ SOverlay::Slot()
+			[
+				SNew(SVerticalBox)
+
+				+ SVerticalBox::Slot()
+				.HAlign(HAlign_Center)
+				.VAlign(VAlign_Center)
+				[
+					SNew(STextBlock)
+					.Text(LOCTEXT("CurveEditorTutorialOverlay", "Select a curve on the left to begin editing."))
+					.Font(FCoreStyle::Get().GetFontStyle("FontAwesome.13"))
+					.ColorAndOpacity(FLinearColor::White)
+				]
+			]
 		];
 	
 
