Curve Editor: Remove ensure in IsPinnedRecursive

#rnx

diff --git a/Engine/Source/Editor/CurveEditor/Private/Tree/SCurveEditorTreePin.cpp b/Engine/Source/Editor/CurveEditor/Private/Tree/SCurveEditorTreePin.cpp
--- a/Engine/Source/Editor/CurveEditor/Private/Tree/SCurveEditorTreePin.cpp
+++ b/Engine/Source/Editor/CurveEditor/Private/Tree/SCurveEditorTreePin.cpp
@@ -108,7 +108,7 @@ void SCurveEditorTreePin::UnpinRecursive(FCurveEditorTreeItemID InTreeItem, FCur
 bool SCurveEditorTreePin::IsPinnedRecursive(FCurveEditorTreeItemID InTreeItem, FCurveEditor* CurveEditor) const
 {
 	const FCurveEditorTreeItem* Item = CurveEditor->FindTreeItem(InTreeItem);
-	if (!ensureMsgf(Item != nullptr, TEXT("Can't find curve editor item. Acting like it's not pinned.")))
+	if (!Item)
 	{
 		return false;
 	}
