Curve Editor: Paste should paste should just paste onto the selected curves (or if none selected, derive the curves to paste onto). It shouldn't paste onto the hovered curves

diff --git a/Engine/Source/Editor/CurveEditor/Private/CurveEditor.cpp b/Engine/Source/Editor/CurveEditor/Private/CurveEditor.cpp
--- a/Engine/Source/Editor/CurveEditor/Private/CurveEditor.cpp
+++ b/Engine/Source/Editor/CurveEditor/Private/CurveEditor.cpp
@@ -1438,49 +1438,29 @@ TSet<FCurveModelID> FCurveEditor::GetTargetCurvesForPaste() const
 {
 	TSet<FCurveModelID> TargetCurves;
 
-	TOptional<FCurveModelID> HoveredID;
-	if (WeakPanel.IsValid())
-	{
-		for (TSharedPtr<SCurveEditorView> View : WeakPanel.Pin()->GetViews())
-		{
-			if (View.IsValid() && View->GetHoveredCurve().IsSet())
-			{
-				HoveredID = View->GetHoveredCurve().GetValue();
-				break;
-			}
-		}
-	}
+	TArray<FCurveEditorTreeItemID> NodesToSearch;
 
-	if (HoveredID.IsSet())
+	// Try nodes with selected keys
+	for (const TTuple<FCurveModelID, FKeyHandleSet>& Pair : Selection.GetAll())
 	{
-		TargetCurves.Add(HoveredID.GetValue());
+		TargetCurves.Add(Pair.Key);
 	}
-	else
-	{
-		TArray<FCurveEditorTreeItemID> NodesToSearch;
 
-		// Try nodes with selected keys
-		for (const TTuple<FCurveModelID, FKeyHandleSet>& Pair : Selection.GetAll())
-		{
-			TargetCurves.Add(Pair.Key);
-		}
-
-		// Try selected nodes
-		if (TargetCurves.Num() == 0)
+	// Try selected nodes
+	if (TargetCurves.Num() == 0)
+	{
+		for (const TTuple<FCurveEditorTreeItemID, ECurveEditorTreeSelectionState>& Pair : GetTreeSelection())
 		{
-			for (const TTuple<FCurveEditorTreeItemID, ECurveEditorTreeSelectionState>& Pair : GetTreeSelection())
-			{
-				NodesToSearch.Add(Pair.Key);
-			}
+			NodesToSearch.Add(Pair.Key);
 		}
+	}
 
-		for (const FCurveEditorTreeItemID& TreeItemID : NodesToSearch)
+	for (const FCurveEditorTreeItemID& TreeItemID : NodesToSearch)
+	{
+		const FCurveEditorTreeItem& TreeItem = GetTreeItem(TreeItemID);
+		for (const FCurveModelID& CurveModelID : TreeItem.GetCurves())
 		{
-			const FCurveEditorTreeItem& TreeItem = GetTreeItem(TreeItemID);
-			for (const FCurveModelID& CurveModelID : TreeItem.GetCurves())
-			{
-				TargetCurves.Add(CurveModelID);
-			}
+			TargetCurves.Add(CurveModelID);
 		}
 	}
 
