Sequencer: Expose TimeBounds which is a subset of the PlaybackRange.

This is the actual playback range that can be restricted by Sequencer with the selection range or the shot range. This allows other systems, like the Curve Editor to be bounded by the same range.

diff --git a/Engine/Source/Editor/CurveEditor/Private/CurveEditor.cpp b/Engine/Source/Editor/CurveEditor/Private/CurveEditor.cpp
--- a/Engine/Source/Editor/CurveEditor/Private/CurveEditor.cpp
+++ b/Engine/Source/Editor/CurveEditor/Private/CurveEditor.cpp
@@ -1018,7 +1018,7 @@ void FCurveEditor::JumpToStart()
 		return;
 	}
 
-	TimeSliderController->SetScrubPosition(TimeSliderController->GetPlayRange().GetLowerBoundValue(), /*bEvaluate*/ true);
+	TimeSliderController->SetScrubPosition(TimeSliderController->GetTimeBounds().GetLowerBoundValue(), /*bEvaluate*/ true);
 }
 
 void FCurveEditor::JumpToEnd()
@@ -1037,7 +1037,7 @@ void FCurveEditor::JumpToEnd()
 	// Calculate an offset from the end to go to. If they have snapping on (and the scrub style is a block) the last valid frame is represented as one
 	// whole display rate frame before the end, otherwise we just subtract a single frame which matches the behavior of hitting play and letting it run to the end.
 	FFrameTime OneFrame = bInsetDisplayFrame ? FFrameRate::TransformTime(FFrameTime(1), DisplayRate, TickResolution) : FFrameTime(1);
-	FFrameTime NewTime = TimeSliderController->GetPlayRange().GetUpperBoundValue() - OneFrame;
+	FFrameTime NewTime = TimeSliderController->GetTimeBounds().GetUpperBoundValue() - OneFrame;
 
 	TimeSliderController->SetScrubPosition(NewTime, /*bEvaluate*/ true);
 }
diff --git a/Engine/Source/Editor/Sequencer/Private/SSequencer.cpp b/Engine/Source/Editor/Sequencer/Private/SSequencer.cpp
--- a/Engine/Source/Editor/Sequencer/Private/SSequencer.cpp
+++ b/Engine/Source/Editor/Sequencer/Private/SSequencer.cpp
@@ -195,6 +195,7 @@ void SSequencer::Construct(const FArguments& InArgs, TSharedRef<FSequencer> InSe
 		TimeSliderArgs.ViewRange = InArgs._ViewRange;
 		TimeSliderArgs.ClampRange = InArgs._ClampRange;
 		TimeSliderArgs.PlaybackRange = InArgs._PlaybackRange;
+		TimeSliderArgs.TimeBounds = InArgs._TimeBounds;
 		TimeSliderArgs.DisplayRate = TAttribute<FFrameRate>(InSequencer, &FSequencer::GetFocusedDisplayRate);
 		TimeSliderArgs.TickResolution = TAttribute<FFrameRate>(InSequencer, &FSequencer::GetFocusedTickResolution);
 		TimeSliderArgs.SelectionRange = InArgs._SelectionRange;
diff --git a/Engine/Source/Editor/Sequencer/Private/SSequencer.h b/Engine/Source/Editor/Sequencer/Private/SSequencer.h
--- a/Engine/Source/Editor/Sequencer/Private/SSequencer.h
+++ b/Engine/Source/Editor/Sequencer/Private/SSequencer.h
@@ -176,6 +176,9 @@ public:
 		/** The playback range */
 		SLATE_ATTRIBUTE( TRange<FFrameNumber>, PlaybackRange )
 
+		/** The time bounds */
+		SLATE_ATTRIBUTE(TRange<FFrameNumber>, TimeBounds)
+
 		/** The selection range */
 		SLATE_ATTRIBUTE( TRange<FFrameNumber>, SelectionRange)
 
diff --git a/Engine/Source/Editor/Sequencer/Private/Sequencer.cpp b/Engine/Source/Editor/Sequencer/Private/Sequencer.cpp
--- a/Engine/Source/Editor/Sequencer/Private/Sequencer.cpp
+++ b/Engine/Source/Editor/Sequencer/Private/Sequencer.cpp
@@ -537,6 +537,7 @@ void FSequencer::InitSequencer(const FSequencerInitParams& InitParams, const TSh
 		.ViewRange( this, &FSequencer::GetViewRange )
 		.ClampRange( this, &FSequencer::GetClampRange )
 		.PlaybackRange( this, &FSequencer::GetPlaybackRange )
+		.TimeBounds( this, &FSequencer::GetTimeBounds )
 		.PlaybackStatus( this, &FSequencer::GetPlaybackStatus )
 		.SelectionRange( this, &FSequencer::GetSelectionRange )
 		.VerticalFrames(this, &FSequencer::GetVerticalFrames)
diff --git a/Engine/Source/Editor/Sequencer/Public/SequencerTimeSliderController.h b/Engine/Source/Editor/Sequencer/Public/SequencerTimeSliderController.h
--- a/Engine/Source/Editor/Sequencer/Public/SequencerTimeSliderController.h
+++ b/Engine/Source/Editor/Sequencer/Public/SequencerTimeSliderController.h
@@ -64,6 +64,9 @@ public:
 	/** Get the current play range for this controller */
 	virtual TRange<FFrameNumber> GetPlayRange() const override { return TimeSliderArgs.PlaybackRange.Get(TRange<FFrameNumber>()); }
 
+	/** Get the time bounds for this controller. The time bounds should be a subset of the playback range. */
+	virtual TRange<FFrameNumber> GetTimeBounds() const override { return TimeSliderArgs.TimeBounds.Get(TRange<FFrameNumber>()); }
+
 	/** Get the selection range */
 	virtual TRange<FFrameNumber> GetSelectionRange() const override { return TimeSliderArgs.SelectionRange.Get(TRange<FFrameNumber>()); }
 
diff --git a/Engine/Source/Editor/SequencerWidgets/Private/ITimeSlider.cpp b/Engine/Source/Editor/SequencerWidgets/Private/ITimeSlider.cpp
--- a/Engine/Source/Editor/SequencerWidgets/Private/ITimeSlider.cpp
+++ b/Engine/Source/Editor/SequencerWidgets/Private/ITimeSlider.cpp
@@ -19,6 +19,11 @@ TRange<FFrameNumber> ITimeSliderController::GetPlayRange() const
 	return TRange<FFrameNumber>();
 }
 
+TRange<FFrameNumber> ITimeSliderController::GetTimeBounds() const
+{
+	return TRange<FFrameNumber>();
+}
+
 TRange<FFrameNumber> ITimeSliderController::GetSelectionRange() const
 {
 	return TRange<FFrameNumber>();
diff --git a/Engine/Source/Editor/SequencerWidgets/Public/ITimeSlider.h b/Engine/Source/Editor/SequencerWidgets/Public/ITimeSlider.h
--- a/Engine/Source/Editor/SequencerWidgets/Public/ITimeSlider.h
+++ b/Engine/Source/Editor/SequencerWidgets/Public/ITimeSlider.h
@@ -124,6 +124,9 @@ public:
 	/** Get the current play range for this controller */
 	SEQUENCERWIDGETS_API virtual TRange<FFrameNumber> GetPlayRange() const;
 
+	/** Get the current time bounds for this controller. The time bounds should be a subset of the playback range. */
+	SEQUENCERWIDGETS_API virtual TRange<FFrameNumber> GetTimeBounds() const;
+
 	/** Get the current selection range for this controller */
 	SEQUENCERWIDGETS_API virtual TRange<FFrameNumber> GetSelectionRange() const;
 
diff --git a/Engine/Source/Editor/SequencerWidgets/Public/TimeSliderArgs.h b/Engine/Source/Editor/SequencerWidgets/Public/TimeSliderArgs.h
--- a/Engine/Source/Editor/SequencerWidgets/Public/TimeSliderArgs.h
+++ b/Engine/Source/Editor/SequencerWidgets/Public/TimeSliderArgs.h
@@ -91,6 +91,9 @@ struct FTimeSliderArgs
 	/** Attribute defining the playback range for this controller */
 	TAttribute<TRange<FFrameNumber>> PlaybackRange;
 
+	/** Attribute defining the time bounds for this controller. The time bounds should be a subset of the playback range. */
+	TAttribute<TRange<FFrameNumber>> TimeBounds;
+
 	/** Attribute for the current sequence's display rate */
 	TAttribute<FFrameRate> DisplayRate;
 
