Curve Editor: Reduce Filter: Crash when reducing a curve with no tangents
#fix Verify curve key attributes contain interp mode before retrieving it, which is the case for bool moviescene channels

diff --git a/Engine/Source/Editor/CurveEditor/Private/Filters/CurveEditorReduceFilter.cpp b/Engine/Source/Editor/CurveEditor/Private/Filters/CurveEditorReduceFilter.cpp
--- a/Engine/Source/Editor/CurveEditor/Private/Filters/CurveEditorReduceFilter.cpp
+++ b/Engine/Source/Editor/CurveEditor/Private/Filters/CurveEditorReduceFilter.cpp
@@ -30,7 +30,7 @@ static float EvalForTwoKeys(const FKeyPosition& Key1Pos, const FKeyAttributes& K
 {
 	const float Diff = Key2Pos.InputValue - Key1Pos.InputValue;
 
-	if (Diff > 0.f && Key1Attrib.GetInterpMode() != RCIM_Constant)
+	if (Diff > 0.f && Key1Attrib.HasInterpMode() && Key1Attrib.GetInterpMode() != RCIM_Constant)
 	{
 		const float Alpha = (InTime - Key1Pos.InputValue) / Diff;
 		const float P0 = Key1Pos.OutputValue;
