Curve Editor: Rework 26837914 to refresh the container on tick rather than in the getters

diff --git a/Engine/Source/Editor/CurveEditor/Private/Views/SInteractiveCurveEditorView.cpp b/Engine/Source/Editor/CurveEditor/Private/Views/SInteractiveCurveEditorView.cpp
--- a/Engine/Source/Editor/CurveEditor/Private/Views/SInteractiveCurveEditorView.cpp
+++ b/Engine/Source/Editor/CurveEditor/Private/Views/SInteractiveCurveEditorView.cpp
@@ -170,7 +170,7 @@ FText SInteractiveCurveEditorView::GetCurveCaption() const
 	if (!CurveCaption.IdenticalTo(CachedCurveCaption))
 	{
 		CachedCurveCaption = CurveCaption;
-		const_cast<SInteractiveCurveEditorView*>(this)->RefreshRetainer();
+		bNeedsRefresh = true;
 	}
 
 	return CurveCaption;
@@ -196,7 +196,7 @@ FSlateColor SInteractiveCurveEditorView::GetCurveCaptionColor() const
 	if (CurveCaptionColor != CachedCurveCaptionColor)
 	{
 		CachedCurveCaptionColor = CurveCaptionColor;
-		const_cast<SInteractiveCurveEditorView*>(this)->RefreshRetainer();
+		bNeedsRefresh = true;
 	}
 
 	return CurveCaptionColor;
@@ -237,6 +237,19 @@ void SInteractiveCurveEditorView::GetGridLinesY(TSharedRef<const FCurveEditor> C
 	}
 }
 
+
+void SInteractiveCurveEditorView::Tick(const FGeometry& AllottedGeometry, const double InCurrentTime, const float InDeltaTime)
+{
+	if (bNeedsRefresh)
+	{
+		bNeedsRefresh = false;
+
+		RefreshRetainer();
+	}
+
+	SCurveEditorView::Tick(AllottedGeometry, InCurrentTime, InDeltaTime);
+}
+
 int32 SInteractiveCurveEditorView::OnPaint(const FPaintArgs& Args, const FGeometry& AllottedGeometry, const FSlateRect& MyCullingRect, FSlateWindowElementList& OutDrawElements, int32 BaseLayerId, const FWidgetStyle& InWidgetStyle, bool bParentEnabled) const
 {
 	PaintView(Args, AllottedGeometry, MyCullingRect, OutDrawElements, BaseLayerId, InWidgetStyle, bParentEnabled);
diff --git a/Engine/Source/Editor/CurveEditor/Public/Views/SInteractiveCurveEditorView.h b/Engine/Source/Editor/CurveEditor/Public/Views/SInteractiveCurveEditorView.h
--- a/Engine/Source/Editor/CurveEditor/Public/Views/SInteractiveCurveEditorView.h
+++ b/Engine/Source/Editor/CurveEditor/Public/Views/SInteractiveCurveEditorView.h
@@ -121,6 +121,7 @@ protected:
 protected:
 
 	// SWidget Interface
+	virtual void Tick(const FGeometry& AllottedGeometry, const double InCurrentTime, const float InDeltaTime) override;
 	virtual int32 OnPaint(const FPaintArgs& Args, const FGeometry& AllottedGeometry, const FSlateRect& MyCullingRect, FSlateWindowElementList& OutDrawElements, int32 BaseLayerId, const FWidgetStyle& InWidgetStyle, bool bParentEnabled) const override;
 	virtual void OnMouseEnter(const FGeometry& MyGeometry, const FPointerEvent& MouseEvent) override;
 	virtual void OnMouseLeave(const FPointerEvent& MouseEvent) override;
@@ -223,4 +224,6 @@ private:
 
 	/** Cached curve caption color, used to determine when to refresh the retainer */
 	mutable FSlateColor CachedCurveCaptionColor;
+
+	mutable bool bNeedsRefresh = false;
 };
\ No newline at end of file
