Curve editor sets new key attributes based on neighbour keys (previous or next key to the curve time)

diff --git a/Engine/Source/Editor/CurveEditor/Private/RichCurveEditorModel.cpp b/Engine/Source/Editor/CurveEditor/Private/RichCurveEditorModel.cpp
--- a/Engine/Source/Editor/CurveEditor/Private/RichCurveEditorModel.cpp
+++ b/Engine/Source/Editor/CurveEditor/Private/RichCurveEditorModel.cpp
@@ -725,6 +725,60 @@ void FRichCurveEditorModel::GetNeighboringKeys(const FKeyHandle InKeyHandle, TOp
 	}
 }
 
+TPair<ERichCurveInterpMode, ERichCurveTangentMode> FRichCurveEditorModel::GetInterpolationMode(const double& InTime, ERichCurveInterpMode DefaultInterpolationMode, ERichCurveTangentMode DefaultTangentMode) const
+{
+	if (IsValid())
+	{
+		const FRichCurve& RichCurve = GetReadOnlyRichCurve();
+
+		if (!RichCurve.Keys.IsEmpty())
+		{
+			FKeyHandle ReferenceKeyHandle;
+			for (auto It = RichCurve.GetKeyHandleIterator(); It; ++It)
+			{
+				if (RichCurve.IsKeyHandleValid(*It))
+				{
+					const FRichCurveKey& Key = RichCurve.GetKeyRef(*It);
+					if (Key.Time < InTime)
+					{
+						ReferenceKeyHandle = *It;
+					}
+					else
+					{
+						// we try to get the key just before the reference time, if it does not exist, we use the key right after
+						if (!RichCurve.IsKeyHandleValid(ReferenceKeyHandle))
+						{
+							ReferenceKeyHandle = *It;
+						}
+						break;
+					}
+				}
+			}
+
+			if (RichCurve.IsKeyHandleValid(ReferenceKeyHandle))
+			{
+				TArray<FKeyAttributes> KeyAttributes;
+				KeyAttributes.SetNum(1);
+				GetKeyAttributes({ ReferenceKeyHandle }, KeyAttributes);
+
+				ERichCurveInterpMode InterpMode = KeyAttributes[0].GetInterpMode();
+				ERichCurveTangentMode TangentMode = KeyAttributes[0].HasTangentMode() ? KeyAttributes[0].GetTangentMode() : DefaultTangentMode;
+				
+				//if we are cubic, with anything but auto tangents we use the default instead, since they will give us flat tangents which aren't good
+				if (InterpMode == ERichCurveInterpMode::RCIM_Cubic &&
+					(TangentMode != ERichCurveTangentMode::RCTM_Auto && TangentMode != ERichCurveTangentMode::RCTM_SmartAuto))
+				{
+					TangentMode = DefaultTangentMode;
+				}
+
+				return TPair<ERichCurveInterpMode, ERichCurveTangentMode>(InterpMode, TangentMode);
+			}
+		}
+	}
+
+	return TPair<ERichCurveInterpMode, ERichCurveTangentMode>(DefaultInterpolationMode, DefaultTangentMode);
+}
+
 FRichCurveEditorModelRaw::FRichCurveEditorModelRaw(FRichCurve* InRichCurve, UObject* InOwner)
 	: FRichCurveEditorModel(InOwner)
 	, RichCurve(InRichCurve)
diff --git a/Engine/Source/Editor/CurveEditor/Private/Views/SInteractiveCurveEditorView.cpp b/Engine/Source/Editor/CurveEditor/Private/Views/SInteractiveCurveEditorView.cpp
--- a/Engine/Source/Editor/CurveEditor/Private/Views/SInteractiveCurveEditorView.cpp
+++ b/Engine/Source/Editor/CurveEditor/Private/Views/SInteractiveCurveEditorView.cpp
@@ -1117,14 +1117,7 @@ FReply SInteractiveCurveEditorView::OnMouseButtonDown(const FGeometry& MyGeometr
 					double MouseTime = CurveSpace.ScreenToSeconds(MousePixel.X);
 					double MouseValue = CurveSpace.ScreenToValue(MousePixel.Y);
 
-					FKeyAttributes KeyAttributes = CurveEditor->GetDefaultKeyAttributes().Get();
-					if (KeyAttributes.HasInterpMode())
-					{
-						// Set interpolation and tangent mode based on surrounding keys, if any
-						TPair<ERichCurveInterpMode, ERichCurveTangentMode> Modes = CurveToAddTo->GetInterpolationMode(MouseTime, KeyAttributes.GetInterpMode(), KeyAttributes.GetTangentMode());
-						KeyAttributes.SetInterpMode(Modes.Key);
-						KeyAttributes.SetTangentMode(Modes.Value);
-					}
+					FKeyAttributes KeyAttributes = GetDefaultKeyAttributesForCurveTime(*CurveEditor, *CurveToAddTo, MouseTime);
 
 					FCurveSnapMetrics SnapMetrics = CurveEditor->GetCurveSnapMetrics(HoveredCurve.GetValue());
 					MouseTime = SnapMetrics.SnapInputSeconds(MouseTime);
@@ -1509,6 +1502,23 @@ bool SInteractiveCurveEditorView::CanApplyBufferedCurves() const
 	return false;
 }
 
+FKeyAttributes SInteractiveCurveEditorView::GetDefaultKeyAttributesForCurveTime(const FCurveEditor& CurveEditor, const FCurveModel& CurveModel, double EvalTime) const
+{
+	FKeyAttributes KeyAttributes = CurveEditor.GetDefaultKeyAttribute().Get();
+
+	TPair<ERichCurveInterpMode, ERichCurveTangentMode> Modes = CurveModel.GetInterpolationMode(EvalTime, ERichCurveInterpMode::RCIM_Linear, ERichCurveTangentMode::RCTM_Auto);
+	if (Modes.Key != ERichCurveInterpMode::RCIM_Linear)
+	{
+		KeyAttributes.SetInterpMode(Modes.Key);
+	}
+	if (Modes.Value != ERichCurveTangentMode::RCTM_Auto)
+	{
+		KeyAttributes.SetTangentMode(Modes.Value);
+	}
+
+	return KeyAttributes;
+}
+
 void SInteractiveCurveEditorView::AddKeyAtScrubTime(TSet<FCurveModelID> ForCurves)
 {
 	TSharedPtr<FCurveEditor> CurveEditor = WeakCurveEditor.Pin();
@@ -1605,14 +1615,7 @@ void SInteractiveCurveEditorView::AddKeyAtTime(const TSet<FCurveModelID>& ToCurv
 			}
 			else
 			{
-				FKeyAttributes KeyAttributes = CurveEditor->GetDefaultKeyAttribute().Get();
-				if (KeyAttributes.HasInterpMode())
-				{
-					// Set interpolation and tangent mode based on surrounding keys, if any
-					TPair<ERichCurveInterpMode, ERichCurveTangentMode> Modes = CurveModel->GetInterpolationMode(EvalTime, KeyAttributes.GetInterpMode(), KeyAttributes.GetTangentMode());
-					KeyAttributes.SetInterpMode(Modes.Key);
-					KeyAttributes.SetTangentMode(Modes.Value);
-				}
+				const FKeyAttributes& KeyAttributes = GetDefaultKeyAttributesForCurveTime(*CurveEditor, *CurveModel, EvalTime);
 
 				// Add a key on this curve
 				NewKey = CurveModel->AddKey(FKeyPosition(EvalTime, CurveValue), KeyAttributes);
diff --git a/Engine/Source/Editor/CurveEditor/Public/RichCurveEditorModel.h b/Engine/Source/Editor/CurveEditor/Public/RichCurveEditorModel.h
--- a/Engine/Source/Editor/CurveEditor/Public/RichCurveEditorModel.h
+++ b/Engine/Source/Editor/CurveEditor/Public/RichCurveEditorModel.h
@@ -53,6 +53,8 @@ public:
 	virtual void GetValueRange(double& MinValue, double& MaxValue) const override;
 	virtual int32 GetNumKeys() const override;
 	virtual void GetNeighboringKeys(const FKeyHandle InKeyHandle, TOptional<FKeyHandle>& OutPreviousKeyHandle, TOptional<FKeyHandle>& OutNextKeyHandle) const override;
+	virtual TPair<ERichCurveInterpMode, ERichCurveTangentMode> GetInterpolationMode(const double& InTime, ERichCurveInterpMode DefaultInterpolationMode, ERichCurveTangentMode DefaultTangentMode) const override;
+
 	virtual bool Evaluate(double ProspectiveTime, double& OutValue) const override;
 	virtual void AddKeys(TArrayView<const FKeyPosition> InKeyPositions, TArrayView<const FKeyAttributes> InAttributes, TArrayView<TOptional<FKeyHandle>>* OutKeyHandles) override;
 	virtual void RemoveKeys(TArrayView<const FKeyHandle> InKeys) override;
diff --git a/Engine/Source/Editor/CurveEditor/Public/Views/SInteractiveCurveEditorView.h b/Engine/Source/Editor/CurveEditor/Public/Views/SInteractiveCurveEditorView.h
--- a/Engine/Source/Editor/CurveEditor/Public/Views/SInteractiveCurveEditorView.h
+++ b/Engine/Source/Editor/CurveEditor/Public/Views/SInteractiveCurveEditorView.h
@@ -42,6 +42,7 @@ struct FCurvePointHandle;
 struct FGeometry;
 struct FOptionalSize;
 struct FPointerEvent;
+struct FKeyAttributes;
 
 namespace CurveViewConstants
 {
@@ -188,6 +189,8 @@ private:
 	bool CanBufferedCurves() const;
 	/** Check if it's legal to apply any of the buffered curves to our currently selected curves. */
 	bool CanApplyBufferedCurves() const;
+	/** Returns interpolation mode and tangent mode based on neighbours or default curve editor if no neighbours . */
+	FKeyAttributes GetDefaultKeyAttributesForCurveTime(const FCurveEditor& CurveEditor, const FCurveModel& CurveModel, double EvalTime) const;
 
 protected:
 
