Sequencer: Deprecrate FCurveModel::GetKeys function that requires FCurveEditor to get ready for new key filter

diff --git a/Engine/Plugins/Animation/ControlRig/Source/ControlRigEditor/Private/ControlRigSpaceChannelCurveModel.cpp b/Engine/Plugins/Animation/ControlRig/Source/ControlRigEditor/Private/ControlRigSpaceChannelCurveModel.cpp
--- a/Engine/Plugins/Animation/ControlRig/Source/ControlRigEditor/Private/ControlRigSpaceChannelCurveModel.cpp
+++ b/Engine/Plugins/Animation/ControlRig/Source/ControlRigEditor/Private/ControlRigSpaceChannelCurveModel.cpp
@@ -161,7 +161,7 @@ void FControlRigSpaceChannelCurveModel::DrawCurve(const FCurveEditor& CurveEdito
 	// Space Channels don't draw any lines so there's no need to fill out the Interpolating Points array.
 }
 
-void FControlRigSpaceChannelCurveModel::GetKeys(const FCurveEditor& CurveEditor, double MinTime, double MaxTime, double MinValue, double MaxValue, TArray<FKeyHandle>& OutKeyHandles) const
+void FControlRigSpaceChannelCurveModel::GetKeys(double MinTime, double MaxTime, double MinValue, double MaxValue, TArray<FKeyHandle>& OutKeyHandles) const
 {
 	FMovieSceneControlRigSpaceChannel* Channel = ChannelHandle.Get();
 	UMovieSceneSection*      Section = WeakSection.Get();
diff --git a/Engine/Plugins/Animation/ControlRig/Source/ControlRigEditor/Private/Tools/ControlRigTweener.cpp b/Engine/Plugins/Animation/ControlRig/Source/ControlRigEditor/Private/Tools/ControlRigTweener.cpp
--- a/Engine/Plugins/Animation/ControlRig/Source/ControlRigEditor/Private/Tools/ControlRigTweener.cpp
+++ b/Engine/Plugins/Animation/ControlRig/Source/ControlRigEditor/Private/Tools/ControlRigTweener.cpp
@@ -128,7 +128,7 @@ void FAnimSliderKeySelection::GetMapOfContiguousKeys()
 		if (const FCurveModel* Curve = CurveEditor->FindCurve(Pair.Key))
 		{
 			AllKeyHandles.Reset();
-			Curve->GetKeys(*CurveEditor, TNumericLimits<double>::Lowest(), TNumericLimits<double>::Max(), TNumericLimits<double>::Lowest(), TNumericLimits<double>::Max(), AllKeyHandles);
+			Curve->GetKeys(TNumericLimits<double>::Lowest(), TNumericLimits<double>::Max(), TNumericLimits<double>::Lowest(), TNumericLimits<double>::Max(), AllKeyHandles);
 			AllKeyPositions.SetNum(AllKeyHandles.Num());
 			Curve->GetKeyPositions(AllKeyHandles, AllKeyPositions);
 			SelectedKeyHandles.Reset(Pair.Value.Num());
diff --git a/Engine/Plugins/Animation/ControlRig/Source/ControlRigEditor/Public/ControlRigSpaceChannelCurveModel.h b/Engine/Plugins/Animation/ControlRig/Source/ControlRigEditor/Public/ControlRigSpaceChannelCurveModel.h
--- a/Engine/Plugins/Animation/ControlRig/Source/ControlRigEditor/Public/ControlRigSpaceChannelCurveModel.h
+++ b/Engine/Plugins/Animation/ControlRig/Source/ControlRigEditor/Public/ControlRigSpaceChannelCurveModel.h
@@ -26,7 +26,7 @@ public:
 	virtual void Modify() override;
 
 	virtual void DrawCurve(const FCurveEditor& CurveEditor, const FCurveEditorScreenSpace& ScreenSpace, TArray<TTuple<double, double>>& InterpolatingPoints) const override;
-	virtual void GetKeys(const FCurveEditor& CurveEditor, double MinTime, double MaxTime, double MinValue, double MaxValue, TArray<FKeyHandle>& OutKeyHandles) const override;
+	virtual void GetKeys(double MinTime, double MaxTime, double MinValue, double MaxValue, TArray<FKeyHandle>& OutKeyHandles) const override;
 	virtual void GetKeyDrawInfo(ECurvePointType PointType, const FKeyHandle InKeyHandle, FKeyDrawInfo& OutDrawInfo) const override;
 
 	virtual void GetKeyPositions(TArrayView<const FKeyHandle> InKeys, TArrayView<FKeyPosition> OutKeyPositions) const override;
diff --git a/Engine/Plugins/Editor/CurveEditorTools/Source/CurveEditorTools/Private/Filters/CurveEditorFFTFilter.cpp b/Engine/Plugins/Editor/CurveEditorTools/Source/CurveEditorTools/Private/Filters/CurveEditorFFTFilter.cpp
--- a/Engine/Plugins/Editor/CurveEditorTools/Source/CurveEditorTools/Private/Filters/CurveEditorFFTFilter.cpp
+++ b/Engine/Plugins/Editor/CurveEditorTools/Source/CurveEditorTools/Private/Filters/CurveEditorFFTFilter.cpp
@@ -56,7 +56,7 @@ void UCurveEditorFFTFilter::ApplyFilter_Impl(TSharedRef<FCurveEditor> InCurveEdi
 
 		// Get all keys that exist between the time range
 		OriginalKeyHandles.Reset();
-		CurveModel->GetKeys(*InCurveEditor, MinKey, MaxKey, TNumericLimits<double>::Lowest(), TNumericLimits<double>::Max(), OriginalKeyHandles);
+		CurveModel->GetKeys(MinKey, MaxKey, TNumericLimits<double>::Lowest(), TNumericLimits<double>::Max(), OriginalKeyHandles);
 		
 		// Need at least three keys to filter
 		if (OriginalKeyHandles.Num() < 3)
diff --git a/Engine/Plugins/Editor/CurveEditorTools/Source/CurveEditorTools/Private/Tools/CurveEditorRetimeTool.cpp b/Engine/Plugins/Editor/CurveEditorTools/Source/CurveEditorTools/Private/Tools/CurveEditorRetimeTool.cpp
--- a/Engine/Plugins/Editor/CurveEditorTools/Source/CurveEditorTools/Private/Tools/CurveEditorRetimeTool.cpp
+++ b/Engine/Plugins/Editor/CurveEditorTools/Source/CurveEditorTools/Private/Tools/CurveEditorRetimeTool.cpp
@@ -471,7 +471,7 @@ void FCurveEditorRetimeTool::OnDragStart()
 			FPreDragChannelData& ChannelData = ResizeData.CurveChannels[ResizeData.CurveChannels.Emplace(CurveID)];
 
 			// Get all Key Handles
-			CurveModel->GetKeys(*CurveEditor, TNumericLimits<double>::Lowest(), TNumericLimits<double>::Max(), TNumericLimits<double>::Lowest(), TNumericLimits<double>::Max(), ChannelData.Handles);
+			CurveModel->GetKeys(TNumericLimits<double>::Lowest(), TNumericLimits<double>::Max(), TNumericLimits<double>::Lowest(), TNumericLimits<double>::Max(), ChannelData.Handles);
 
 			// Get all Key Positions for those Handles.
 			ChannelData.FrameNumbers.SetNumUninitialized(ChannelData.Handles.Num());
diff --git a/Engine/Plugins/Editor/CurveEditorTools/Source/CurveEditorTools/Private/Tools/CurveEditorTransformTool.cpp b/Engine/Plugins/Editor/CurveEditorTools/Source/CurveEditorTools/Private/Tools/CurveEditorTransformTool.cpp
--- a/Engine/Plugins/Editor/CurveEditorTools/Source/CurveEditorTools/Private/Tools/CurveEditorTransformTool.cpp
+++ b/Engine/Plugins/Editor/CurveEditorTools/Source/CurveEditorTools/Private/Tools/CurveEditorTransformTool.cpp
@@ -1096,7 +1096,7 @@ void FCurveEditorTransformTool::OnDragEnd()
 			for (const FKeyPosition& KeyTime : KeyTimes)
 			{
 				TArray<FKeyHandle> KeysInRange;
-				Curve->GetKeys(*CurveEditor, KeyTime.InputValue, KeyTime.InputValue, TNumericLimits<double>::Lowest(), TNumericLimits<double>::Max(), KeysInRange);
+				Curve->GetKeys(KeyTime.InputValue, KeyTime.InputValue, TNumericLimits<double>::Lowest(), TNumericLimits<double>::Max(), KeysInRange);
 
 				// If there's more than 1 key at this time, remove all but the keys that moved the largest amount
 				if (KeysInRange.Num() > 1)
diff --git a/Engine/Plugins/Experimental/Mutable/Source/CustomizableObjectEditor/Private/MuCOE/SMutableCurveViewer.h b/Engine/Plugins/Experimental/Mutable/Source/CustomizableObjectEditor/Private/MuCOE/SMutableCurveViewer.h
--- a/Engine/Plugins/Experimental/Mutable/Source/CustomizableObjectEditor/Private/MuCOE/SMutableCurveViewer.h
+++ b/Engine/Plugins/Experimental/Mutable/Source/CustomizableObjectEditor/Private/MuCOE/SMutableCurveViewer.h
@@ -49,7 +49,7 @@ public:
 	// Based on RefineCurvePoints from RichCurveEditorModel.cpp
 	void RefineCurvePoints(const FRichCurve& InRichCurve, double TimeThreshold, float ValueThreshold, TArray<TTuple<double, double>>& InOutPoints) const;
 
-	virtual void GetKeys(const FCurveEditor& CurveEditor, double MinTime, double MaxTime, double MinValue,
+	virtual void GetKeys(double MinTime, double MaxTime, double MinValue,
 	                     double MaxValue, TArray<FKeyHandle>& OutKeyHandles) const override {}
 	
 	virtual void AddKeys(TArrayView<const FKeyPosition> InPositions, TArrayView<const FKeyAttributes> InAttributes,
diff --git a/Engine/Plugins/FX/Niagara/Source/NiagaraEditorWidgets/Private/DetailCustomizations/NiagaraDataInterfaceCurveDetails.cpp b/Engine/Plugins/FX/Niagara/Source/NiagaraEditorWidgets/Private/DetailCustomizations/NiagaraDataInterfaceCurveDetails.cpp
--- a/Engine/Plugins/FX/Niagara/Source/NiagaraEditorWidgets/Private/DetailCustomizations/NiagaraDataInterfaceCurveDetails.cpp
+++ b/Engine/Plugins/FX/Niagara/Source/NiagaraEditorWidgets/Private/DetailCustomizations/NiagaraDataInterfaceCurveDetails.cpp
@@ -372,7 +372,7 @@ private:
 	void GetSortedKeyHandlessAndPositionsForModel(FCurveModel& InCurveModel, TArray<FKeyHandlePositionPair>& OutSortedKeyHandlesAndPositions)
 	{
 		TArray<FKeyHandle> KeyHandles;
-		InCurveModel.GetKeys(*CurveEditor.Get(), TNumericLimits<double>::Lowest(), TNumericLimits<double>::Max(), TNumericLimits<double>::Lowest(), TNumericLimits<double>::Max(), KeyHandles);
+		InCurveModel.GetKeys(TNumericLimits<double>::Lowest(), TNumericLimits<double>::Max(), TNumericLimits<double>::Lowest(), TNumericLimits<double>::Max(), KeyHandles);
 
 		TArray<FKeyPosition> KeyPositions;
 		KeyPositions.AddDefaulted(KeyHandles.Num());
@@ -655,7 +655,7 @@ private:
 			{
 				// If there's a single key, add the new key at the same value, but time + 1.
 				TArray<FKeyHandle> KeyHandles;
-				CurveModelForAdd->GetKeys(*CurveEditor.Get(), TNumericLimits<double>::Lowest(), TNumericLimits<double>::Max(), TNumericLimits<double>::Lowest(), TNumericLimits<double>::Max(), KeyHandles);
+				CurveModelForAdd->GetKeys(TNumericLimits<double>::Lowest(), TNumericLimits<double>::Max(), TNumericLimits<double>::Lowest(), TNumericLimits<double>::Max(), KeyHandles);
 				
 				TArray<FKeyPosition> KeyPositions;
 				KeyPositions.AddDefaulted();
@@ -1549,7 +1549,7 @@ FReply FNiagaraDataInterfaceCurveDetailsBase::CurveTemplateSelected(TWeakObjectP
 				if (CurveModel != nullptr)
 				{
 					TArray<FKeyHandle> KeyHandles;
-					CurveModel->GetKeys(*CurveEditor.Get(), TNumericLimits<double>::Lowest(), TNumericLimits<double>::Max(), TNumericLimits<double>::Lowest(), TNumericLimits<double>::Max(), KeyHandles);
+					CurveModel->GetKeys(TNumericLimits<double>::Lowest(), TNumericLimits<double>::Max(), TNumericLimits<double>::Lowest(), TNumericLimits<double>::Max(), KeyHandles);
 					CurveModel->RemoveKeys(KeyHandles);
 
 					const FRichCurve& FloatCurve = FloatCurveAsset->FloatCurve;
diff --git a/Engine/Source/Editor/CurveEditor/Private/CurveEditor.cpp b/Engine/Source/Editor/CurveEditor/Private/CurveEditor.cpp
--- a/Engine/Source/Editor/CurveEditor/Private/CurveEditor.cpp
+++ b/Engine/Source/Editor/CurveEditor/Private/CurveEditor.cpp
@@ -882,7 +882,7 @@ void FCurveEditor::StepToNextKey()
 		{
 			TArray<FKeyHandle> KeyHandles;
 			double MaxTime = NextTime.IsSet() ? NextTime.GetValue() : TNumericLimits<double>::Max();
-			CurveModel->GetKeys(*this, CurrentTime, MaxTime, TNumericLimits<double>::Lowest(), TNumericLimits<double>::Max(), KeyHandles);
+			CurveModel->GetKeys(CurrentTime, MaxTime, TNumericLimits<double>::Lowest(), TNumericLimits<double>::Max(), KeyHandles);
 
 			TArray<FKeyPosition> KeyPositions;
 			KeyPositions.SetNum(KeyHandles.Num());
@@ -941,7 +941,7 @@ void FCurveEditor::StepToPreviousKey()
 		{
 			TArray<FKeyHandle> KeyHandles;
 			double MinTime = PreviousTime.IsSet() ? PreviousTime.GetValue() : TNumericLimits<double>::Lowest();
-			CurveModel->GetKeys(*this, MinTime, CurrentTime, TNumericLimits<double>::Lowest(), TNumericLimits<double>::Max(), KeyHandles);
+			CurveModel->GetKeys(MinTime, CurrentTime, TNumericLimits<double>::Lowest(), TNumericLimits<double>::Max(), KeyHandles);
 
 			TArray<FKeyPosition> KeyPositions;
 			KeyPositions.SetNum(KeyHandles.Num());
@@ -1100,7 +1100,7 @@ void FCurveEditor::SelectAllKeys()
 		if (FCurveModel* Curve = FindCurve(ID))
 		{
 			TArray<FKeyHandle> KeyHandles;
-			Curve->GetKeys(*this, TNumericLimits<double>::Lowest(), TNumericLimits<double>::Max(), TNumericLimits<double>::Lowest(), TNumericLimits<double>::Max(), KeyHandles);
+			Curve->GetKeys(TNumericLimits<double>::Lowest(), TNumericLimits<double>::Max(), TNumericLimits<double>::Lowest(), TNumericLimits<double>::Max(), KeyHandles);
 			Selection.Add(ID, ECurvePointType::Key, KeyHandles);
 		}
 	}
@@ -1125,7 +1125,7 @@ void FCurveEditor::SelectForward()
 		if (FCurveModel* Curve = FindCurve(ID))
 		{
 			TArray<FKeyHandle> KeyHandles;
-			Curve->GetKeys(*this, CurrentTime, TNumericLimits<double>::Max(), TNumericLimits<double>::Lowest(), TNumericLimits<double>::Max(), KeyHandles);
+			Curve->GetKeys(CurrentTime, TNumericLimits<double>::Max(), TNumericLimits<double>::Lowest(), TNumericLimits<double>::Max(), KeyHandles);
 			Selection.Add(ID, ECurvePointType::Key, KeyHandles);
 		}
 	}
@@ -1150,7 +1150,7 @@ void FCurveEditor::SelectBackward()
 		if (FCurveModel* Curve = FindCurve(ID))
 		{
 			TArray<FKeyHandle> KeyHandles;
-			Curve->GetKeys(*this, TNumericLimits<double>::Min(), CurrentTime, TNumericLimits<double>::Lowest(), TNumericLimits<double>::Max(), KeyHandles);
+			Curve->GetKeys(TNumericLimits<double>::Min(), CurrentTime, TNumericLimits<double>::Lowest(), TNumericLimits<double>::Max(), KeyHandles);
 			Selection.Add(ID, ECurvePointType::Key, KeyHandles);
 		}
 	}
@@ -1502,7 +1502,7 @@ bool FCurveEditor::CopyBufferCurveToCurveID(const UCurveEditorCopyableCurveKeys*
 		// Just double checking we actually set a Min/Max time so we don't wipe out every key to infinity.
 		if (InSourceCurve->KeyPositions.Num() > 0)
 		{
-			TargetCurve->GetKeys(*this, MinKeyTime, MaxKeyTime, TNumericLimits<double>::Lowest(), TNumericLimits<double>::Max(), KeysToRemove);
+			TargetCurve->GetKeys(MinKeyTime, MaxKeyTime, TNumericLimits<double>::Lowest(), TNumericLimits<double>::Max(), KeysToRemove);
 		}
 
 		TargetCurve->RemoveKeys(KeysToRemove);
@@ -2045,7 +2045,7 @@ void FCurveEditor::ApplyBufferedCurveToTarget(const IBufferedCurveModel* Buffere
 
 	// Copy the data from the Buffered curve into the target curve. This just does wholesale replacement.
 	TArray<FKeyHandle> TargetKeyHandles;
-	TargetCurve->GetKeys(*this, TNumericLimits<double>::Lowest(), TNumericLimits<double>::Max(), TNumericLimits<double>::Lowest(), TNumericLimits<double>::Max(), TargetKeyHandles);
+	TargetCurve->GetKeys(TNumericLimits<double>::Lowest(), TNumericLimits<double>::Max(), TNumericLimits<double>::Lowest(), TNumericLimits<double>::Max(), TargetKeyHandles);
 
 	// Clear our current keys from the target curve
 	TargetCurve->RemoveKeys(TargetKeyHandles);
@@ -2310,7 +2310,7 @@ void FCurveEditor::PostUndo(bool bSuccess)
 		}
 		// Get all of the key handles from this curve.
 		TArray<FKeyHandle> KeyHandles;
-		CurveModel->GetKeys(*this, TNumericLimits<double>::Lowest(), TNumericLimits<double>::Max(), TNumericLimits<double>::Lowest(), TNumericLimits<double>::Max(), KeyHandles);
+		CurveModel->GetKeys(TNumericLimits<double>::Lowest(), TNumericLimits<double>::Max(), TNumericLimits<double>::Lowest(), TNumericLimits<double>::Max(), KeyHandles);
 
 		// The set handles will be mutated as we remove things so we need a copy that we can iterate through.
 		TArrayView<const FKeyHandle> SelectedHandles = Set.Value.AsArray();
diff --git a/Engine/Source/Editor/CurveEditor/Private/DragOperations/CurveEditorDragOperation_MoveKeys.cpp b/Engine/Source/Editor/CurveEditor/Private/DragOperations/CurveEditorDragOperation_MoveKeys.cpp
--- a/Engine/Source/Editor/CurveEditor/Private/DragOperations/CurveEditorDragOperation_MoveKeys.cpp
+++ b/Engine/Source/Editor/CurveEditor/Private/DragOperations/CurveEditorDragOperation_MoveKeys.cpp
@@ -175,7 +175,7 @@ void FCurveEditorDragOperation_MoveKeys::OnEndDrag(FVector2D InitialPosition, FV
 			for (const FKeyPosition& KeyTime : KeyTimes)
 			{
 				TArray<FKeyHandle> KeysInRange;
-				Curve->GetKeys(*CurveEditor, KeyTime.InputValue, KeyTime.InputValue, TNumericLimits<double>::Lowest(), TNumericLimits<double>::Max(), KeysInRange);
+				Curve->GetKeys(KeyTime.InputValue, KeyTime.InputValue, TNumericLimits<double>::Lowest(), TNumericLimits<double>::Max(), KeysInRange);
 
 				// If there's more than 1 key at this time, remove all but the keys that moved the largest amount
 				if (KeysInRange.Num() > 1)
diff --git a/Engine/Source/Editor/CurveEditor/Private/Filters/CurveEditorBakeFilter.cpp b/Engine/Source/Editor/CurveEditor/Private/Filters/CurveEditorBakeFilter.cpp
--- a/Engine/Source/Editor/CurveEditor/Private/Filters/CurveEditorBakeFilter.cpp
+++ b/Engine/Source/Editor/CurveEditor/Private/Filters/CurveEditorBakeFilter.cpp
@@ -133,7 +133,7 @@ void UCurveEditorBakeFilter::ApplyFilter_Impl(TSharedRef<FCurveEditor> InCurveEd
 
 		// Get all keys that exist between the time range
 		KeyHandles.Reset();
-		Curve->GetKeys(*InCurveEditor, MinKey, MaxKey, TNumericLimits<double>::Lowest(), TNumericLimits<double>::Max(), KeyHandles);
+		Curve->GetKeys(MinKey, MaxKey, TNumericLimits<double>::Lowest(), TNumericLimits<double>::Max(), KeyHandles);
 
 		// Determine new times for new keys
 		
diff --git a/Engine/Source/Editor/CurveEditor/Private/Filters/CurveEditorEulerFilter.cpp b/Engine/Source/Editor/CurveEditor/Private/Filters/CurveEditorEulerFilter.cpp
--- a/Engine/Source/Editor/CurveEditor/Private/Filters/CurveEditorEulerFilter.cpp
+++ b/Engine/Source/Editor/CurveEditor/Private/Filters/CurveEditorEulerFilter.cpp
@@ -57,7 +57,7 @@ void UCurveEditorEulerFilter::ApplyFilter_Impl(TSharedRef<FCurveEditor> InCurveE
 
 		// Get all keys that exist between the time range
 		KeyHandles.Reset();
-		Curve->GetKeys(*InCurveEditor, MinKey, MaxKey, TNumericLimits<double>::Lowest(), TNumericLimits<double>::Max(), KeyHandles);
+		Curve->GetKeys(MinKey, MaxKey, TNumericLimits<double>::Lowest(), TNumericLimits<double>::Max(), KeyHandles);
 
 		if (KeyHandles.Num() >= 2)
 		{
diff --git a/Engine/Source/Editor/CurveEditor/Private/Filters/CurveEditorReduceFilter.cpp b/Engine/Source/Editor/CurveEditor/Private/Filters/CurveEditorReduceFilter.cpp
--- a/Engine/Source/Editor/CurveEditor/Private/Filters/CurveEditorReduceFilter.cpp
+++ b/Engine/Source/Editor/CurveEditor/Private/Filters/CurveEditorReduceFilter.cpp
@@ -93,7 +93,7 @@ void UCurveEditorReduceFilter::ApplyFilter_Impl(TSharedRef<FCurveEditor> InCurve
 
 		// Get all keys that exist between the time range
 		KeyHandles.Reset();
-		Curve->GetKeys(*InCurveEditor, MinKey, MaxKey, TNumericLimits<double>::Lowest(), TNumericLimits<double>::Max(), KeyHandles);
+		Curve->GetKeys( MinKey, MaxKey, TNumericLimits<double>::Lowest(), TNumericLimits<double>::Max(), KeyHandles);
 		if (KeyHandles.Num() > 2) //need at least 3 keys to reduce
 		{
 			SelectedKeyPositions.SetNum(KeyHandles.Num());
@@ -182,4 +182,4 @@ void UCurveEditorReduceFilter::ApplyFilter_Impl(TSharedRef<FCurveEditor> InCurve
 			Curve->RemoveKeys(KeysToRemove);
 		}
 	}
-}
\ No newline at end of file
+}
diff --git a/Engine/Source/Editor/CurveEditor/Private/RichCurveEditorModel.cpp b/Engine/Source/Editor/CurveEditor/Private/RichCurveEditorModel.cpp
--- a/Engine/Source/Editor/CurveEditor/Private/RichCurveEditorModel.cpp
+++ b/Engine/Source/Editor/CurveEditor/Private/RichCurveEditorModel.cpp
@@ -252,7 +252,7 @@ void FRichCurveEditorModel::DrawCurve(const FCurveEditor& CurveEditor, const FCu
 	}
 }
 
-void FRichCurveEditorModel::GetKeys(const FCurveEditor& CurveEditor, double MinTime, double MaxTime, double MinValue, double MaxValue, TArray<FKeyHandle>& OutKeyHandles) const
+void FRichCurveEditorModel::GetKeys(double MinTime, double MaxTime, double MinValue, double MaxValue, TArray<FKeyHandle>& OutKeyHandles) const
 {
 	if (UObject* Owner = WeakOwner.Get())
 	{
diff --git a/Engine/Source/Editor/CurveEditor/Private/SCurveEditorView.cpp b/Engine/Source/Editor/CurveEditor/Private/SCurveEditorView.cpp
--- a/Engine/Source/Editor/CurveEditor/Private/SCurveEditorView.cpp
+++ b/Engine/Source/Editor/CurveEditor/Private/SCurveEditorView.cpp
@@ -229,7 +229,7 @@ void SCurveEditorView::GetCurveDrawParam(TSharedPtr<FCurveEditor>& CurveEditor,c
 	}
 
 	TArray<FKeyHandle> VisibleKeys;
-	CurveModel->GetKeys(*CurveEditor, InputMin, InputMax, TNumericLimits<double>::Lowest(), TNumericLimits<double>::Max(), VisibleKeys);
+	CurveModel->GetKeys(InputMin, InputMax, TNumericLimits<double>::Lowest(), TNumericLimits<double>::Max(), VisibleKeys);
 
 	// Always reset the points to cover case going from 1 to 0 keys
 	Params.Points.Reset(VisibleKeys.Num());
diff --git a/Engine/Source/Editor/CurveEditor/Private/Tree/SCurveEditorTreeSelect.cpp b/Engine/Source/Editor/CurveEditor/Private/Tree/SCurveEditorTreeSelect.cpp
--- a/Engine/Source/Editor/CurveEditor/Private/Tree/SCurveEditorTreeSelect.cpp
+++ b/Engine/Source/Editor/CurveEditor/Private/Tree/SCurveEditorTreeSelect.cpp
@@ -70,7 +70,7 @@ FReply SCurveEditorTreeSelect::SelectAll()
 		{
 			TArray<FKeyHandle> KeyHandles;
 			KeyHandles.Reserve(CurveModel->GetNumKeys());
-			CurveModel->GetKeys(*CurveEditor.Get(), TNumericLimits<double>::Lowest(), TNumericLimits<double>::Max(), TNumericLimits<double>::Lowest(), TNumericLimits<double>::Max(), KeyHandles);
+			CurveModel->GetKeys(TNumericLimits<double>::Lowest(), TNumericLimits<double>::Max(), TNumericLimits<double>::Lowest(), TNumericLimits<double>::Max(), KeyHandles);
 		
 			if (bIsControlDown)
 			{
diff --git a/Engine/Source/Editor/CurveEditor/Private/Views/SInteractiveCurveEditorView.cpp b/Engine/Source/Editor/CurveEditor/Private/Views/SInteractiveCurveEditorView.cpp
--- a/Engine/Source/Editor/CurveEditor/Private/Views/SInteractiveCurveEditorView.cpp
+++ b/Engine/Source/Editor/CurveEditor/Private/Views/SInteractiveCurveEditorView.cpp
@@ -681,7 +681,7 @@ bool SInteractiveCurveEditorView::GetCurveWithinWidgetRange(const FSlateRect& Wi
 		if (const FCurveModel* Curve = CurveEditor->FindCurve(CurveID))
 		{
 			TArray<FKeyHandle> KeyHandles;
-			Curve->GetKeys(*CurveEditor, TNumericLimits<double>::Lowest(), TNumericLimits<double>::Max(), TNumericLimits<double>::Lowest(), TNumericLimits<double>::Max(), KeyHandles);
+			Curve->GetKeys(TNumericLimits<double>::Lowest(), TNumericLimits<double>::Max(), TNumericLimits<double>::Lowest(), TNumericLimits<double>::Max(), KeyHandles);
 
 			for (const FKeyHandle& KeyHandle : KeyHandles)
 			{
@@ -1283,7 +1283,7 @@ FReply SInteractiveCurveEditorView::OnMouseButtonUp(const FGeometry& MyGeometry,
 
 			TArray<FKeyHandle> KeyHandles;
 			KeyHandles.Reserve(CurveModel->GetNumKeys());
-			CurveModel->GetKeys(*CurveEditor, TNumericLimits<double>::Lowest(), TNumericLimits<double>::Max(), TNumericLimits<double>::Lowest(), TNumericLimits<double>::Max(), KeyHandles);
+			CurveModel->GetKeys(TNumericLimits<double>::Lowest(), TNumericLimits<double>::Max(), TNumericLimits<double>::Lowest(), TNumericLimits<double>::Max(), KeyHandles);
 
 			// Add or remove all keys from the curve.
 			if (bIsShiftDown)
@@ -1605,7 +1605,7 @@ void SInteractiveCurveEditorView::AddKeyAtTime(const TSet<FCurveModelID>& ToCurv
 			// add that key to the selection set instead. This solves issues with snapping causing keys to be created adjacent
 			// to the mouse cursor (sometimes by a large amount).
 			TArray<FKeyHandle> ExistingKeys;
-			CurveModel->GetKeys(*CurveEditor, EvalTime - KINDA_SMALL_NUMBER, EvalTime + KINDA_SMALL_NUMBER, TNumericLimits<double>::Lowest(), TNumericLimits<double>::Max(), ExistingKeys);
+			CurveModel->GetKeys(EvalTime - KINDA_SMALL_NUMBER, EvalTime + KINDA_SMALL_NUMBER, TNumericLimits<double>::Lowest(), TNumericLimits<double>::Max(), ExistingKeys);
 			
 			TOptional<FKeyHandle> NewKey;
 
diff --git a/Engine/Source/Editor/CurveEditor/Public/CurveModel.h b/Engine/Source/Editor/CurveEditor/Public/CurveModel.h
--- a/Engine/Source/Editor/CurveEditor/Public/CurveModel.h
+++ b/Engine/Source/Editor/CurveEditor/Public/CurveModel.h
@@ -72,17 +72,23 @@ public:
 	 */
 	virtual void DrawCurve(const FCurveEditor& CurveEditor, const FCurveEditorScreenSpace& ScreenSpace, TArray<TTuple<double, double>>& InterpolatingPoints) const = 0;
 
+	UE_DEPRECATED(5.5, "Implement the version that doesn't take the Curve Editor parameter.")
+	virtual void GetKeys(const FCurveEditor& CurveEditor, double MinTime, double MaxTime, double MinValue, double MaxValue, TArray<FKeyHandle>& OutKeyHandles) const
+	{
+		GetKeys(MinTime, MaxTime, MinValue, MaxValue, OutKeyHandles);
+	}
+
+
 	/**
 	 * Retrieve all keys that lie in the specified time and value range
 	 *
-	 * @param CurveEditor             Reference to the curve editor that is retrieving keys.
 	 * @param MinTime                 Minimum key time to return in seconds
 	 * @param MaxTime                 Maximum key time to return in seconds
 	 * @param MinValue                Minimum key value to return
 	 * @param MaxValue                Maximum key value to return
 	 * @param OutKeyHandles           Array to populate with key handles that reside within the specified ranges
 	 */
-	virtual void GetKeys(const FCurveEditor& CurveEditor, double MinTime, double MaxTime, double MinValue, double MaxValue, TArray<FKeyHandle>& OutKeyHandles) const = 0;
+	virtual void GetKeys(double MinTime, double MaxTime, double MinValue, double MaxValue, TArray<FKeyHandle>& OutKeyHandles) const = 0;
 
 	/**
 	 * Add keys to this curve
diff --git a/Engine/Source/Editor/CurveEditor/Public/Filters/CurveEditorFilterBase.h b/Engine/Source/Editor/CurveEditor/Public/Filters/CurveEditorFilterBase.h
--- a/Engine/Source/Editor/CurveEditor/Public/Filters/CurveEditorFilterBase.h
+++ b/Engine/Source/Editor/CurveEditor/Public/Filters/CurveEditorFilterBase.h
@@ -110,8 +110,8 @@ public:
 			FKeyHandleSet& HandleSet = KeysToOperateOn.Add(CurveModelID);
 			
 			TArray<FKeyHandle> Handles;
-			CurveModel->GetKeys(*InCurveEditor, TNumericLimits<double>::Lowest(), TNumericLimits<double>::Max(), TNumericLimits<double>::Lowest(), TNumericLimits<double>::Max(), Handles);
-			
+			CurveModel->GetKeys(TNumericLimits<double>::Lowest(), TNumericLimits<double>::Max(), TNumericLimits<double>::Lowest(), TNumericLimits<double>::Max(), Handles);
+
 			for (const FKeyHandle& Handle : Handles)
 			{
 				HandleSet.Add(Handle, ECurvePointType::Key);
diff --git a/Engine/Source/Editor/CurveEditor/Public/KeyBarCurveModel.h b/Engine/Source/Editor/CurveEditor/Public/KeyBarCurveModel.h
--- a/Engine/Source/Editor/CurveEditor/Public/KeyBarCurveModel.h
+++ b/Engine/Source/Editor/CurveEditor/Public/KeyBarCurveModel.h
@@ -36,7 +36,7 @@ public:
 	virtual const void* GetCurve() const override { return nullptr; }
 	virtual void Modify() override {};
 	virtual void DrawCurve(const FCurveEditor& CurveEditor, const FCurveEditorScreenSpace& ScreenSpace, TArray<TTuple<double, double>>& InterpolatingPoints) const override{};
-	virtual void GetKeys(const FCurveEditor& CurveEditor, double MinTime, double MaxTime, double MinValue, double MaxValue, TArray<FKeyHandle>& OutKeyHandles) const override {};
+	virtual void GetKeys(double MinTime, double MaxTime, double MinValue, double MaxValue, TArray<FKeyHandle>& OutKeyHandles) const override {};
 	virtual void AddKeys(TArrayView<const FKeyPosition> InPositions, TArrayView<const FKeyAttributes> InAttributes, TArrayView<TOptional<FKeyHandle>>* OutKeyHandles = nullptr) override {};
 	virtual void RemoveKeys(TArrayView<const FKeyHandle> InKeys) override {};
 	virtual void GetKeyPositions(TArrayView<const FKeyHandle> InKeys, TArrayView<FKeyPosition> OutKeyPositions) const override {};
diff --git a/Engine/Source/Editor/CurveEditor/Public/RichCurveEditorModel.h b/Engine/Source/Editor/CurveEditor/Public/RichCurveEditorModel.h
--- a/Engine/Source/Editor/CurveEditor/Public/RichCurveEditorModel.h
+++ b/Engine/Source/Editor/CurveEditor/Public/RichCurveEditorModel.h
@@ -38,7 +38,7 @@ public:
 	virtual void Modify() override;
 
 	virtual void DrawCurve(const FCurveEditor& CurveEditor, const FCurveEditorScreenSpace& ScreenSpace, TArray<TTuple<double, double>>& InterpolatingPoints) const override;
-	virtual void GetKeys(const FCurveEditor& CurveEditor, double MinTime, double MaxTime, double MinValue, double MaxValue, TArray<FKeyHandle>& OutKeyHandles) const override;
+	virtual void GetKeys(double MinTime, double MaxTime, double MinValue, double MaxValue, TArray<FKeyHandle>& OutKeyHandles) const override;
 	virtual void GetKeyDrawInfo(ECurvePointType PointType, const FKeyHandle InKeyHandle, FKeyDrawInfo& OutDrawInfo) const override;
 
 	virtual void GetKeyPositions(TArrayView<const FKeyHandle> InKeys, TArrayView<FKeyPosition> OutKeyPositions) const override;
diff --git a/Engine/Source/Editor/CurveTableEditor/Private/RealCurveModel.cpp b/Engine/Source/Editor/CurveTableEditor/Private/RealCurveModel.cpp
--- a/Engine/Source/Editor/CurveTableEditor/Private/RealCurveModel.cpp
+++ b/Engine/Source/Editor/CurveTableEditor/Private/RealCurveModel.cpp
@@ -83,7 +83,7 @@ void FRealCurveModel::DrawCurve(const FCurveEditor& CurveEditor, const FCurveEdi
 	}
 }
 
-void FRealCurveModel::GetKeys(const FCurveEditor& CurveEditor, double MinTime, double MaxTime, double MinValue, double MaxValue, TArray<FKeyHandle>& OutKeyHandles) const
+void FRealCurveModel::GetKeys(double MinTime, double MaxTime, double MinValue, double MaxValue, TArray<FKeyHandle>& OutKeyHandles) const
 {
 	if (IsValid())
 	{
diff --git a/Engine/Source/Editor/CurveTableEditor/Private/RealCurveModel.h b/Engine/Source/Editor/CurveTableEditor/Private/RealCurveModel.h
--- a/Engine/Source/Editor/CurveTableEditor/Private/RealCurveModel.h
+++ b/Engine/Source/Editor/CurveTableEditor/Private/RealCurveModel.h
@@ -38,7 +38,7 @@ class FRealCurveModel : public FCurveModel
 	virtual void Modify() override;
 
 	virtual void DrawCurve(const FCurveEditor& CurveEditor, const FCurveEditorScreenSpace& ScreenSpace, TArray<TTuple<double, double>>& InterpolatingPoints) const override;
-	virtual void GetKeys(const FCurveEditor& CurveEditor, double MinTime, double MaxTime, double MinValue, double MaxValue, TArray<FKeyHandle>& OutKeyHandles) const override;
+	virtual void GetKeys(double MinTime, double MaxTime, double MinValue, double MaxValue, TArray<FKeyHandle>& OutKeyHandles) const override;
 	virtual void AddKeys(TArrayView<const FKeyPosition> InKeyPositions, TArrayView<const FKeyAttributes> InAttributes, TArrayView<TOptional<FKeyHandle>>* OutKeyHandles) override;
 	virtual void RemoveKeys(TArrayView<const FKeyHandle> InKeys) override;
 
diff --git a/Engine/Source/Editor/MovieSceneTools/Private/Channels/ChannelCurveModel.cpp b/Engine/Source/Editor/MovieSceneTools/Private/Channels/ChannelCurveModel.cpp
--- a/Engine/Source/Editor/MovieSceneTools/Private/Channels/ChannelCurveModel.cpp
+++ b/Engine/Source/Editor/MovieSceneTools/Private/Channels/ChannelCurveModel.cpp
@@ -153,7 +153,7 @@ void FChannelCurveModel<ChannelType, ChannelValue, KeyType>::DrawCurve(const FCu
 }
 
 template <class ChannelType, class ChannelValue, class KeyType>
-void FChannelCurveModel<ChannelType, ChannelValue, KeyType>::GetKeys(const FCurveEditor& CurveEditor, double MinTime, double MaxTime, double MinValue, double MaxValue, TArray<FKeyHandle>& OutKeyHandles) const
+void FChannelCurveModel<ChannelType, ChannelValue, KeyType>::GetKeys(double MinTime, double MaxTime, double MinValue, double MaxValue, TArray<FKeyHandle>& OutKeyHandles) const
 {
 	ChannelType* Channel = ChannelHandle.Get();
 	UMovieSceneSection* Section = WeakSection.Get();
diff --git a/Engine/Source/Editor/MovieSceneTools/Private/Channels/ConstraintChannelCurveModel.cpp b/Engine/Source/Editor/MovieSceneTools/Private/Channels/ConstraintChannelCurveModel.cpp
--- a/Engine/Source/Editor/MovieSceneTools/Private/Channels/ConstraintChannelCurveModel.cpp
+++ b/Engine/Source/Editor/MovieSceneTools/Private/Channels/ConstraintChannelCurveModel.cpp
@@ -186,7 +186,7 @@ void FConstraintChannelCurveModel::DrawCurve(const FCurveEditor& CurveEditor, co
 	// Constraint Channels don't draw any lines so there's no need to fill out the Interpolating Points array.
 }
 
-void FConstraintChannelCurveModel::GetKeys(const FCurveEditor& CurveEditor, double MinTime, double MaxTime, double MinValue, double MaxValue, TArray<FKeyHandle>& OutKeyHandles) const
+void FConstraintChannelCurveModel::GetKeys(double MinTime, double MaxTime, double MinValue, double MaxValue, TArray<FKeyHandle>& OutKeyHandles) const
 {
 	FMovieSceneConstraintChannel* Channel = GetChannel();
 	UMovieSceneSection* Section = WeakSection.Get();
diff --git a/Engine/Source/Editor/MovieSceneTools/Private/Channels/EventChannelCurveModel.cpp b/Engine/Source/Editor/MovieSceneTools/Private/Channels/EventChannelCurveModel.cpp
--- a/Engine/Source/Editor/MovieSceneTools/Private/Channels/EventChannelCurveModel.cpp
+++ b/Engine/Source/Editor/MovieSceneTools/Private/Channels/EventChannelCurveModel.cpp
@@ -121,7 +121,7 @@ void FEventChannelCurveModel::DrawCurve(const FCurveEditor& CurveEditor, const F
 	// Event Channels don't draw any lines so there's no need to fill out the Interpolating Points array.
 }
 
-void FEventChannelCurveModel::GetKeys(const FCurveEditor& CurveEditor, double MinTime, double MaxTime, double MinValue, double MaxValue, TArray<FKeyHandle>& OutKeyHandles) const
+void FEventChannelCurveModel::GetKeys(double MinTime, double MaxTime, double MinValue, double MaxValue, TArray<FKeyHandle>& OutKeyHandles) const
 {
 	FMovieSceneEventChannel* Channel = ChannelHandle.Get();
 	UMovieSceneSection*      Section = WeakSection.Get();
diff --git a/Engine/Source/Editor/MovieSceneTools/Private/Channels/EventChannelCurveModel.h b/Engine/Source/Editor/MovieSceneTools/Private/Channels/EventChannelCurveModel.h
--- a/Engine/Source/Editor/MovieSceneTools/Private/Channels/EventChannelCurveModel.h
+++ b/Engine/Source/Editor/MovieSceneTools/Private/Channels/EventChannelCurveModel.h
@@ -39,7 +39,7 @@ public:
 	virtual void Modify() override;
 
 	virtual void DrawCurve(const FCurveEditor& CurveEditor, const FCurveEditorScreenSpace& ScreenSpace, TArray<TTuple<double, double>>& InterpolatingPoints) const override;
-	virtual void GetKeys(const FCurveEditor& CurveEditor, double MinTime, double MaxTime, double MinValue, double MaxValue, TArray<FKeyHandle>& OutKeyHandles) const override;
+	virtual void GetKeys(double MinTime, double MaxTime, double MinValue, double MaxValue, TArray<FKeyHandle>& OutKeyHandles) const override;
 	virtual void GetKeyDrawInfo(ECurvePointType PointType, const FKeyHandle InKeyHandle, FKeyDrawInfo& OutDrawInfo) const override;
 
 	virtual void GetKeyPositions(TArrayView<const FKeyHandle> InKeys, TArrayView<FKeyPosition> OutKeyPositions) const override;
diff --git a/Engine/Source/Editor/MovieSceneTools/Public/Channels/ChannelCurveModel.h b/Engine/Source/Editor/MovieSceneTools/Public/Channels/ChannelCurveModel.h
--- a/Engine/Source/Editor/MovieSceneTools/Public/Channels/ChannelCurveModel.h
+++ b/Engine/Source/Editor/MovieSceneTools/Public/Channels/ChannelCurveModel.h
@@ -47,7 +47,7 @@ public:
 	virtual void Modify() override;
 
 	virtual void DrawCurve(const FCurveEditor& CurveEditor, const FCurveEditorScreenSpace& ScreenSpace, TArray<TTuple<double, double>>& OutInterpolatingPoints) const override;
-	virtual void GetKeys(const FCurveEditor& CurveEditor, double MinTime, double MaxTime, double MinValue, double MaxValue, TArray<FKeyHandle>& OutKeyHandles) const override;
+	virtual void GetKeys(double MinTime, double MaxTime, double MinValue, double MaxValue, TArray<FKeyHandle>& OutKeyHandles) const override;
 	virtual void GetKeyDrawInfo(ECurvePointType PointType, const FKeyHandle InKeyHandle, FKeyDrawInfo& OutDrawInfo) const override;
 
 	virtual void GetKeyPositions(TArrayView<const FKeyHandle> InKeys, TArrayView<FKeyPosition> OutKeyPositions) const override;
diff --git a/Engine/Source/Editor/MovieSceneTools/Public/Channels/ConstraintChannelCurveModel.h b/Engine/Source/Editor/MovieSceneTools/Public/Channels/ConstraintChannelCurveModel.h
--- a/Engine/Source/Editor/MovieSceneTools/Public/Channels/ConstraintChannelCurveModel.h
+++ b/Engine/Source/Editor/MovieSceneTools/Public/Channels/ConstraintChannelCurveModel.h
@@ -45,7 +45,7 @@ public:
 	virtual void Modify() override;
 
 	virtual void DrawCurve(const FCurveEditor& CurveEditor, const FCurveEditorScreenSpace& ScreenSpace, TArray<TTuple<double, double>>& InterpolatingPoints) const override;
-	virtual void GetKeys(const FCurveEditor& CurveEditor, double MinTime, double MaxTime, double MinValue, double MaxValue, TArray<FKeyHandle>& OutKeyHandles) const override;
+	virtual void GetKeys(double MinTime, double MaxTime, double MinValue, double MaxValue, TArray<FKeyHandle>& OutKeyHandles) const override;
 	virtual void GetKeyDrawInfo(ECurvePointType PointType, const FKeyHandle InKeyHandle, FKeyDrawInfo& OutDrawInfo) const override;
 
 	virtual void GetKeyPositions(TArrayView<const FKeyHandle> InKeys, TArrayView<FKeyPosition> OutKeyPositions) const override;
