Sequencer: Control Rig: New Smart Baking algorithm added as a curve filter, make new filter default for reducing control rig keys when baking, add bp/python support. Note we still support the old baking via UMovieScene::Optimize, etc.. in the code, may want to revisit changing that algorithm over once the new one is battle tested.

diff --git a/Engine/Plugins/Animation/ControlRig/Source/ControlRigEditor/Private/BakeToControlRigSettings.cpp b/Engine/Plugins/Animation/ControlRig/Source/ControlRigEditor/Private/BakeToControlRigSettings.cpp
--- a/Engine/Plugins/Animation/ControlRig/Source/ControlRigEditor/Private/BakeToControlRigSettings.cpp
+++ b/Engine/Plugins/Animation/ControlRig/Source/ControlRigEditor/Private/BakeToControlRigSettings.cpp
@@ -11,6 +11,6 @@ UBakeToControlRigSettings::UBakeToControlRigSettings(const FObjectInitializer& I
 void UBakeToControlRigSettings::Reset()
 {
 	bReduceKeys = false;
-	Tolerance = 0.001f;
+	SmartReduce.Reset();
 }
 
diff --git a/Engine/Plugins/Animation/ControlRig/Source/ControlRigEditor/Private/BakeToControlRigSettings.h b/Engine/Plugins/Animation/ControlRig/Source/ControlRigEditor/Private/BakeToControlRigSettings.h
--- a/Engine/Plugins/Animation/ControlRig/Source/ControlRigEditor/Private/BakeToControlRigSettings.h
+++ b/Engine/Plugins/Animation/ControlRig/Source/ControlRigEditor/Private/BakeToControlRigSettings.h
@@ -4,6 +4,7 @@
 #include "CoreMinimal.h"
 #include "UObject/ObjectMacros.h"
 #include "UObject/Object.h"
+#include "Filters/CurveEditorSmartReduceFilter.h"
 #include "BakeToControlRigSettings.generated.h"
 
 UCLASS(BlueprintType, config = EditorSettings)
@@ -18,9 +19,13 @@ public:
 	bool bReduceKeys = false;
 
 	/** Reduce Keys Tolerance*/
-	UPROPERTY(EditAnywhere, Category = "Reduce Keys")
+	UE_DEPRECATED(5.5, "Use the SmartReduce parameter instead.")
+	UPROPERTY()
 	float Tolerance = 0.001f;
 
+	UPROPERTY(EditAnywhere, Category = "Reduce Keys", meta = (EditCondition = "bReduceKeys"))
+	FSmartReduceParams SmartReduce;
+
 	/** Reset controls to initial value on every frame */
 	UPROPERTY(EditAnywhere, Category = "Reset Controls")
 	bool bResetControls = true;
diff --git a/Engine/Plugins/Animation/ControlRig/Source/ControlRigEditor/Private/ControlRigEditorModule.cpp b/Engine/Plugins/Animation/ControlRig/Source/ControlRigEditor/Private/ControlRigEditorModule.cpp
--- a/Engine/Plugins/Animation/ControlRig/Source/ControlRigEditor/Private/ControlRigEditorModule.cpp
+++ b/Engine/Plugins/Animation/ControlRig/Source/ControlRigEditor/Private/ControlRigEditorModule.cpp
@@ -1144,7 +1144,7 @@ void FControlRigEditorModule::BakeToControlRig(UClass* ControlRigClass, UAnimSeq
 
 		if (WeakSequencer.IsValid())
 		{
-
+			TSharedPtr<ISequencer> SequencerPtr = WeakSequencer.Pin();
 			ASkeletalMeshActor* MeshActor = World->SpawnActor<ASkeletalMeshActor>(ASkeletalMeshActor::StaticClass(), FTransform::Identity);
 			MeshActor->SetActorLabel(AnimSequence->GetName());
 
@@ -1159,18 +1159,18 @@ void FControlRigEditorModule::BakeToControlRig(UClass* ControlRigClass, UAnimSeq
 			MeshActor->RegisterAllComponents();
 			TArray<TWeakObjectPtr<AActor> > ActorsToAdd;
 			ActorsToAdd.Add(MeshActor);
-			TArray<FGuid> ActorTracks = WeakSequencer.Pin()->AddActors(ActorsToAdd, false);
+			TArray<FGuid> ActorTracks = SequencerPtr->AddActors(ActorsToAdd, false);
 			FGuid ActorTrackGuid = ActorTracks[0];
 
 			// By default, convert this to a spawnable and delete the existing actor. If for some reason, 
 			// the spawnable couldn't be generated, use the existing actor as a possessable (this could 
 			// eventually be an option)
-			TArray<FGuid> SpawnableGuids = WeakSequencer.Pin()->ConvertToSpawnable(ActorTrackGuid);
+			TArray<FGuid> SpawnableGuids = SequencerPtr->ConvertToSpawnable(ActorTrackGuid);
 			if (SpawnableGuids.Num())
 			{	
 				ActorTrackGuid = SpawnableGuids[0];
 
-				UObject* SpawnedMesh = WeakSequencer.Pin()->FindSpawnedObjectOrTemplate(ActorTrackGuid);
+				UObject* SpawnedMesh = SequencerPtr->FindSpawnedObjectOrTemplate(ActorTrackGuid);
 
 				if (SpawnedMesh)
 				{
@@ -1186,7 +1186,7 @@ void FControlRigEditorModule::BakeToControlRig(UClass* ControlRigClass, UAnimSeq
 
 			//Delete binding from default animating rig
 			//if we have skel mesh component binding we can just delete that
-			FGuid CompGuid = WeakSequencer.Pin()->FindObjectId(*(MeshActor->GetSkeletalMeshComponent()), WeakSequencer.Pin()->GetFocusedTemplateID());
+			FGuid CompGuid = SequencerPtr->FindObjectId(*(MeshActor->GetSkeletalMeshComponent()), SequencerPtr->GetFocusedTemplateID());
 			if (CompGuid.IsValid())
 			{
 				if (!MovieScene->RemovePossessable(CompGuid))
@@ -1218,7 +1218,7 @@ void FControlRigEditorModule::BakeToControlRig(UClass* ControlRigClass, UAnimSeq
 				ControlRig->Initialize();
 				ControlRig->Evaluate_AnyThread();
 
-				WeakSequencer.Pin()->NotifyMovieSceneDataChanged(EMovieSceneDataChangeType::MovieSceneStructureItemsChanged);
+				SequencerPtr->NotifyMovieSceneDataChanged(EMovieSceneDataChangeType::MovieSceneStructureItemsChanged);
 
 				Track->Modify();
 				UMovieSceneSection* NewSection = Track->CreateControlRigSection(0, ControlRig, true);
@@ -1226,22 +1226,23 @@ void FControlRigEditorModule::BakeToControlRig(UClass* ControlRigClass, UAnimSeq
 				Track->SetTrackName(FName(*ObjectName));
 				Track->SetDisplayName(FText::FromString(ObjectName));
 				UMovieSceneControlRigParameterSection* ParamSection = Cast<UMovieSceneControlRigParameterSection>(NewSection);
-			
-				FBakeToControlDelegate BakeCallback = FBakeToControlDelegate::CreateLambda([this, WeakSequencer, LevelSequence, 
+				FBakeToControlDelegate BakeCallback = FBakeToControlDelegate::CreateLambda([this, &SequencerPtr, LevelSequence,
 					AnimSequence, MovieScene, ControlRig, ParamSection,ActorTrackGuid, SkelMeshComp]
-				(bool bKeyReduce, float KeyReduceTolerance, bool bResetControls)
+				(bool bKeyReduce, float KeyReduceTolerance, FFrameRate BakeFrameRate, bool bResetControls)
 				{
 					if (ParamSection)
 					{
-						EMovieSceneKeyInterpolation DefaultInterpolation = WeakSequencer.Pin()->GetKeyInterpolation();
-						ParamSection->LoadAnimSequenceIntoThisSection(AnimSequence, MovieScene, SkelMeshComp, bKeyReduce,
-							KeyReduceTolerance, bResetControls, FFrameNumber(0), DefaultInterpolation);
+						FSmartReduceParams SmartReduce;
+						SmartReduce.TolerancePercentage = KeyReduceTolerance;
+						SmartReduce.SampleRate = BakeFrameRate;
+						FControlRigParameterTrackEditor::LoadAnimationIntoSection(SequencerPtr, AnimSequence, SkelMeshComp, FFrameNumber(0),
+							bKeyReduce, SmartReduce,bResetControls, ParamSection);
 					}
-					WeakSequencer.Pin()->EmptySelection();
-					WeakSequencer.Pin()->SelectSection(ParamSection);
-					WeakSequencer.Pin()->ThrobSectionSelection();
-					WeakSequencer.Pin()->ObjectImplicitlyAdded(ControlRig);
-					WeakSequencer.Pin()->NotifyMovieSceneDataChanged(EMovieSceneDataChangeType::MovieSceneStructureItemAdded);
+					SequencerPtr->EmptySelection();
+					SequencerPtr->SelectSection(ParamSection);
+					SequencerPtr->ThrobSectionSelection();
+					SequencerPtr->ObjectImplicitlyAdded(ControlRig);
+					SequencerPtr->NotifyMovieSceneDataChanged(EMovieSceneDataChangeType::MovieSceneStructureItemAdded);
 					FControlRigEditMode* ControlRigEditMode = static_cast<FControlRigEditMode*>(GLevelEditorModeTools().GetActiveMode(FControlRigEditMode::ModeName));
 					if (!ControlRigEditMode)
 					{
@@ -1250,7 +1251,7 @@ void FControlRigEditorModule::BakeToControlRig(UClass* ControlRigClass, UAnimSeq
 					}
 					if (ControlRigEditMode)
 					{
-						ControlRigEditMode->AddControlRigObject(ControlRig, WeakSequencer.Pin());
+						ControlRigEditMode->AddControlRigObject(ControlRig, SequencerPtr);
 					}
 
 					//create soft links to each other
diff --git a/Engine/Plugins/Animation/ControlRig/Source/ControlRigEditor/Private/ControlRigSequencerEditorLibrary.cpp b/Engine/Plugins/Animation/ControlRig/Source/ControlRigEditor/Private/ControlRigSequencerEditorLibrary.cpp
--- a/Engine/Plugins/Animation/ControlRig/Source/ControlRigEditor/Private/ControlRigSequencerEditorLibrary.cpp
+++ b/Engine/Plugins/Animation/ControlRig/Source/ControlRigEditor/Private/ControlRigSequencerEditorLibrary.cpp
@@ -1103,6 +1103,26 @@ void UControlRigSequencerEditorLibrary::SetControlRigWorldTransforms(ULevelSeque
 	LocalSetControlRigWorldTransforms(LevelSequence, ControlRig, ControlName, EControlRigSetKey::Always, Frames, WorldTransforms,TimeUnit);
 }
 
+bool UControlRigSequencerEditorLibrary::SmartReduce(FSmartReduceParams& ReduceParams, UMovieSceneSection* MovieSceneSection)
+{
+	//get level sequence if one exists...
+	TWeakPtr<ISequencer> WeakSequencer = GetSequencerFromAsset();
+	if (WeakSequencer.IsValid() == false)
+	{
+		UE_LOG(LogControlRig, Error, TEXT("Smart Reduce: No open level sequence"));
+		return false;
+	}
+	if (UMovieSceneControlRigParameterSection* Section = Cast<UMovieSceneControlRigParameterSection>(MovieSceneSection))
+	{
+		TSharedPtr<ISequencer>  SequencerPtr = WeakSequencer.Pin();
+		FControlRigParameterTrackEditor::SmartReduce(SequencerPtr, ReduceParams, Section);
+		return true;
+	}
+	UE_LOG(LogControlRig, Error, TEXT("Smart Reduce: Section is not Control Rig Section"));
+	return false;
+}
+
+
 bool UControlRigSequencerEditorLibrary::BakeToControlRig(UWorld* World, ULevelSequence* LevelSequence, UClass* InClass, UAnimSeqExportOption* ExportOptions, bool bReduceKeys, float Tolerance,
 	const FMovieSceneBindingProxy& Binding, bool bResetControls)
 {
@@ -1135,9 +1155,6 @@ bool UControlRigSequencerEditorLibrary::BakeToControlRig(UWorld* World, ULevelSe
 		UE_LOG(LogControlRig, Error, TEXT("Baking: Problem Setting up Player"));
 		return false;
 	}
-	
-	
-
 	bool bResult = false;
 	const FScopedTransaction Transaction(LOCTEXT("BakeToControlRig_Transaction", "Bake To Control Rig"));
 	{
diff --git a/Engine/Plugins/Animation/ControlRig/Source/ControlRigEditor/Private/SBakeToControlRigDialog.cpp b/Engine/Plugins/Animation/ControlRig/Source/ControlRigEditor/Private/SBakeToControlRigDialog.cpp
--- a/Engine/Plugins/Animation/ControlRig/Source/ControlRigEditor/Private/SBakeToControlRigDialog.cpp
+++ b/Engine/Plugins/Animation/ControlRig/Source/ControlRigEditor/Private/SBakeToControlRigDialog.cpp
@@ -74,7 +74,7 @@ private:
 		TSharedPtr<SWindow> Window = FSlateApplication::Get().FindWidgetWindow(AsShared());
 		if (BakeSettings && Delegate.IsBound())
 		{
-			Delegate.Execute(BakeSettings->bReduceKeys, BakeSettings->Tolerance, BakeSettings->bResetControls);
+			Delegate.Execute(BakeSettings->bReduceKeys, BakeSettings->SmartReduce.TolerancePercentage, BakeSettings->SmartReduce.SampleRate, BakeSettings->bResetControls);
 		}
 		if (Window.IsValid())
 		{
diff --git a/Engine/Plugins/Animation/ControlRig/Source/ControlRigEditor/Private/SBakeToControlRigDialog.h b/Engine/Plugins/Animation/ControlRig/Source/ControlRigEditor/Private/SBakeToControlRigDialog.h
--- a/Engine/Plugins/Animation/ControlRig/Source/ControlRigEditor/Private/SBakeToControlRigDialog.h
+++ b/Engine/Plugins/Animation/ControlRig/Source/ControlRigEditor/Private/SBakeToControlRigDialog.h
@@ -4,11 +4,11 @@
 #include "CoreMinimal.h"
 #include "Delegates/Delegate.h"
 #include "Widgets/SWindow.h"
+#include "Misc/FrameRate.h"
 
-DECLARE_DELEGATE_ThreeParams(FBakeToControlDelegate, bool, float, bool);
+DECLARE_DELEGATE_FourParams(FBakeToControlDelegate, bool, float, FFrameRate, bool);
 
 struct BakeToControlRigDialog
 {
 	static void GetBakeParams(FBakeToControlDelegate& Delegate, const FOnWindowClosed& OnClosedDelegate);
-
 };
\ No newline at end of file
diff --git a/Engine/Plugins/Animation/ControlRig/Source/ControlRigEditor/Private/Sequencer/ControlRigParameterTrackEditor.cpp b/Engine/Plugins/Animation/ControlRig/Source/ControlRigEditor/Private/Sequencer/ControlRigParameterTrackEditor.cpp
--- a/Engine/Plugins/Animation/ControlRig/Source/ControlRigEditor/Private/Sequencer/ControlRigParameterTrackEditor.cpp
+++ b/Engine/Plugins/Animation/ControlRig/Source/ControlRigEditor/Private/Sequencer/ControlRigParameterTrackEditor.cpp
@@ -100,7 +100,8 @@
 #include "Widgets/Layout/SSpacer.h"
 #include "ControlRigSequencerEditorLibrary.h"
 #include "LevelSequence.h"
-
+#include "ISequencerModule.h"
+#include "CurveModel.h"
 
 #define LOCTEXT_NAMESPACE "FControlRigParameterTrackEditor"
 
@@ -955,6 +956,51 @@ FReply SBakeToAnimAndControlRigOptionsWindow::OnResetToDefaultClick() const
 	return FReply::Handled();
 }
 
+void FControlRigParameterTrackEditor::SmartReduce(TSharedPtr<ISequencer>& InSequencer, const FSmartReduceParams& InParams, UMovieSceneControlRigParameterSection * ParamSection)
+{
+	if (ParamSection)
+	{
+		ISequencerModule& SequencerModule = FModuleManager::LoadModuleChecked<ISequencerModule>("Sequencer");
+		const bool bNeedToTestExisting = false;
+		TOptional<FKeyHandleSet> KeyHandleSet;
+
+		FMovieSceneChannelProxy& ChannelProxy = ParamSection->GetChannelProxy();
+		for (const FMovieSceneChannelEntry& Entry : ParamSection->GetChannelProxy().GetAllEntries())
+		{
+			const FName ChannelTypeName = Entry.GetChannelTypeName();
+			TArrayView<FMovieSceneChannel* const> Channels = Entry.GetChannels();
+			for (int32 Index = 0; Index < Channels.Num(); ++Index)
+			{
+				FMovieSceneChannelHandle ChannelHandle = ChannelProxy.MakeHandle(ChannelTypeName, Index);
+				ISequencerChannelInterface* EditorInterface = SequencerModule.FindChannelEditorInterface(ChannelHandle.GetChannelTypeName());
+				if (TUniquePtr<FCurveModel> CurveModel = EditorInterface->CreateCurveEditorModel_Raw(ChannelHandle, ParamSection, InSequencer.ToSharedRef()))
+				{
+					FKeyHandleSet OutHandleSet;
+					UCurveEditorSmartReduceFilter::SmartReduce(CurveModel.Get(), InParams, KeyHandleSet, bNeedToTestExisting, OutHandleSet);
+				}
+			}
+		}
+	}
+}
+
+bool FControlRigParameterTrackEditor::LoadAnimationIntoSection(TSharedPtr<ISequencer>& SequencerPtr, UAnimSequence* AnimSequence, USkeletalMeshComponent* SkelMeshComp,
+	FFrameNumber StartFrame, bool bReduceKeys, const FSmartReduceParams& ReduceParams, bool bResetControls, UMovieSceneControlRigParameterSection* ParamSection)
+{
+	EMovieSceneKeyInterpolation DefaultInterpolation = SequencerPtr->GetKeyInterpolation();
+	UMovieSceneSequence* OwnerSequence = SequencerPtr->GetFocusedMovieSceneSequence();
+	UMovieScene* OwnerMovieScene = OwnerSequence->GetMovieScene();
+	if (ParamSection->LoadAnimSequenceIntoThisSection(AnimSequence, OwnerMovieScene, SkelMeshComp,
+		false, 0.0, bResetControls, StartFrame, DefaultInterpolation))
+	{
+		if (bReduceKeys)
+		{
+			SmartReduce(SequencerPtr, ReduceParams, ParamSection);
+		}
+		return true;
+	}
+	return false;
+}
+
 void FControlRigParameterTrackEditor::BakeToControlRig(UClass* InClass, FGuid ObjectBinding, UObject* BoundActor, USkeletalMeshComponent* SkelMeshComp, USkeleton* Skeleton)
 {
 	FSlateApplication::Get().DismissAllMenus();
@@ -1087,10 +1133,9 @@ void FControlRigParameterTrackEditor::BakeToControlRig(UClass* InClass, FGuid Ob
 					GetSequencer()->SelectSection(NewSection);
 					GetSequencer()->ThrobSectionSelection();
 					GetSequencer()->NotifyMovieSceneDataChanged(EMovieSceneDataChangeType::MovieSceneStructureItemAdded);
-					EMovieSceneKeyInterpolation DefaultInterpolation = SequencerParent->GetKeyInterpolation();
-					ParamSection->LoadAnimSequenceIntoThisSection(TempAnimSequence, OwnerMovieScene, SkelMeshComp,
-						BakeSettings->bReduceKeys, BakeSettings->Tolerance, BakeSettings->bResetControls, FFrameNumber(0), DefaultInterpolation);
-
+					TSharedPtr<ISequencer> SequencerPtr = GetSequencer();
+					LoadAnimationIntoSection(SequencerPtr, TempAnimSequence, SkelMeshComp, FFrameNumber(0),
+						BakeSettings->bReduceKeys, BakeSettings->SmartReduce, BakeSettings->bResetControls, ParamSection);
 					//Turn Off Any Skeletal Animation Tracks
 					TArray<UMovieSceneSkeletalAnimationTrack*> SkelAnimationTracks;
 					if (const FMovieSceneBinding* Binding = OwnerMovieScene->FindBinding(ObjectBinding))
@@ -1197,7 +1242,8 @@ void FControlRigParameterTrackEditor::BakeInvertedPose(UControlRig* InControlRig
 	UMovieSceneSequence* MovieSceneSequence = GetSequencer()->GetFocusedMovieSceneSequence();
 	UMovieScene* MovieScene = MovieSceneSequence->GetMovieScene();
 	UAnimSeqExportOption* ExportOptions = NewObject<UAnimSeqExportOption>(GetTransientPackage(), NAME_None);
-	UBakeToControlRigSettings* BakeSettings = GetMutableDefault<UBakeToControlRigSettings>();
+	//@sara to do, not sure if you want to key reduce after, but BakeSettings isn't used
+	//UBakeToControlRigSettings* BakeSettings = GetMutableDefault<UBakeToControlRigSettings>();
 	const TSharedPtr<ISequencer> ParentSequencer = GetSequencer();
 	FMovieSceneSequenceIDRef Template = ParentSequencer->GetFocusedTemplateID();
 	FMovieSceneSequenceTransform RootToLocalTransform = ParentSequencer->GetFocusedMovieSceneSequenceTransform();
@@ -5283,10 +5329,10 @@ void FControlRigParameterSection::OnAnimationAssetSelectedForFK(const FAssetData
 
 			FScopedTransaction Transaction(LOCTEXT("BakeAnimation_Transaction", "Bake Animation To FK Control Rig"));
 			Section->Modify();
-			UMovieScene* MovieScene = SequencerPtr->GetFocusedMovieSceneSequence()->GetMovieScene();
 			FFrameNumber StartFrame = SequencerPtr->GetLocalTime().Time.GetFrame();
-			EMovieSceneKeyInterpolation DefaultInterpolation = SequencerPtr->GetKeyInterpolation();
-			if (!Section->LoadAnimSequenceIntoThisSection(AnimSequence, MovieScene, SkelMeshComp, false, 0.1f, true, StartFrame, DefaultInterpolation))
+			FSmartReduceParams SmartReduce;
+			if(!FControlRigParameterTrackEditor::LoadAnimationIntoSection(SequencerPtr, AnimSequence, SkelMeshComp,StartFrame,
+				false, SmartReduce, true, Section))
 			{
 				Transaction.Cancel();
 			}
diff --git a/Engine/Plugins/Animation/ControlRig/Source/ControlRigEditor/Private/Sequencer/ControlRigParameterTrackEditor.h b/Engine/Plugins/Animation/ControlRig/Source/ControlRigEditor/Private/Sequencer/ControlRigParameterTrackEditor.h
--- a/Engine/Plugins/Animation/ControlRig/Source/ControlRigEditor/Private/Sequencer/ControlRigParameterTrackEditor.h
+++ b/Engine/Plugins/Animation/ControlRig/Source/ControlRigEditor/Private/Sequencer/ControlRigParameterTrackEditor.h
@@ -18,6 +18,7 @@
 #include "MovieSceneToolsModule.h"
 #include "Engine/EngineTypes.h"
 #include "EditorUndoClient.h"
+#include "Filters/CurveEditorSmartReduceFilter.h"
 #include "ScopedTransaction.h"
 
 class UTickableTransformConstraint;
@@ -242,6 +243,12 @@ public:
 	//**Function to collapse all layers from this section onto the first absoluate layer.*/
 	static bool CollapseAllLayers(TSharedPtr<ISequencer>& SequencerPtr, UMovieSceneTrack* OwnerTrack, const FBakingAnimationKeySettings& InSettings);
 
+	//** Function to load animation into a section, returns true if successful
+	static bool LoadAnimationIntoSection(TSharedPtr<ISequencer>& SequencerPtr,  UAnimSequence* AnimSequence, USkeletalMeshComponent* SkelMeshComp,
+		FFrameNumber StartFrame, bool bReduceKeys, const FSmartReduceParams& ReduceParams, bool bResetControls, UMovieSceneControlRigParameterSection* ParamSection);
+
+	//** Function to smart reduce all keys on all constraints based on the parameters
+	static void SmartReduce(TSharedPtr<ISequencer>& SequencerPtr, const FSmartReduceParams& InParams, UMovieSceneControlRigParameterSection* InSection);
 private:
 	FDelegateHandle SelectionChangedHandle;
 	FDelegateHandle SequencerChangedHandle;
diff --git a/Engine/Plugins/Animation/ControlRig/Source/ControlRigEditor/Public/ControlRigSequencerEditorLibrary.h b/Engine/Plugins/Animation/ControlRig/Source/ControlRigEditor/Public/ControlRigSequencerEditorLibrary.h
--- a/Engine/Plugins/Animation/ControlRig/Source/ControlRigEditor/Public/ControlRigSequencerEditorLibrary.h
+++ b/Engine/Plugins/Animation/ControlRig/Source/ControlRigEditor/Public/ControlRigSequencerEditorLibrary.h
@@ -17,6 +17,7 @@
 #include "Tools/ControlRigSnapSettings.h"
 #include "MovieSceneTimeUnit.h"
 #include "RigSpacePickerBakeSettings.h"
+#include "Filters/CurveEditorSmartReduceFilter.h"
 #include "ControlRigSequencerEditorLibrary.generated.h"
 
 class ULevelSequence;
@@ -146,6 +147,19 @@ public:
 	static bool BakeToControlRig(UWorld* World, ULevelSequence* LevelSequence, UClass* ControlRigClass, UAnimSeqExportOption* ExportOptions, bool bReduceKeys, float Tolerance,
 			const FMovieSceneBindingProxy& Binding, bool bResetControls = true);
 
+	/**
+	* Peform new Smart Reduce filter over the specified control rig section in the current open level sequence. Note existing
+	* functions like LoadAnimSequenceIntoControlRigSection and BakeToControlRig, will still use the old key reduction algorithm,
+	* so if you want to bake and then key reduce with the new function, set the bKeyReduce param as false with those functions,
+	* but then call this function after.
+	* @param ReduceParams Key reduction parameters
+	* @param MovieSceneSection The Control rig section we want to reduce
+	* @return returns True if successful, False otherwise
+	*/
+	UFUNCTION(BlueprintCallable, Category = "Editor Scripting | Sequencer Tools | Control Rig")
+	static bool SmartReduce(FSmartReduceParams& ReduceParams, UMovieSceneSection* MovieSceneSection);
+
+
 	/**
 	* Bake the constraint to keys based on the passed in frames. This will use the open sequencer to bake. See ConstraintsScriptingLibrary to get the list of available constraints
 	* @param World The active world
diff --git a/Engine/Source/Editor/CurveEditor/Private/Filters/CurveEditorSmartReduceFilter.cpp b/Engine/Source/Editor/CurveEditor/Private/Filters/CurveEditorSmartReduceFilter.cpp
new file mode 100644
--- /dev/null
+++ b/Engine/Source/Editor/CurveEditor/Private/Filters/CurveEditorSmartReduceFilter.cpp
@@ -0,0 +1,318 @@
+// Copyright Epic Games, Inc. All Rights Reserved.
+
+#include "Filters/CurveEditorSmartReduceFilter.h"
+
+#include "Containers/Array.h"
+#include "Containers/ArrayView.h"
+#include "CurveDataAbstraction.h"
+#include "CurveEditor.h"
+#include "CurveEditorSelection.h"
+#include "CurveEditorSnapMetrics.h"
+#include "CurveEditorTypes.h"
+#include "CurveModel.h"
+#include "Curves/CurveEvaluation.h"
+#include "Curves/KeyHandle.h"
+#include "Curves/RealCurve.h"
+#include "Curves/RichCurve.h"
+#include "HAL/PlatformCrt.h"
+#include "Math/NumericLimits.h"
+#include "Math/UnrealMathSSE.h"
+#include "Misc/FrameNumber.h"
+#include "Misc/FrameRate.h"
+#include "Misc/FrameTime.h"
+#include "Templates/SharedPointer.h"
+#include "Templates/Tuple.h"
+#include "Templates/UnrealTemplate.h"
+
+
+// cache of times and positions based on the frame rate.
+struct FCurveCache
+{
+	double Step = 0.1;
+	double Min = TNumericLimits<double>::Max();
+	double Max = TNumericLimits<double>::Lowest();
+
+	TArray<double> Times;
+	TArray<double> Positions;
+
+	double Evaluate(double Time)
+	{	
+		if (Positions.Num() <= 0)
+		{
+			return 0.0;
+		}
+		int32 Index = GetIndex(Time);
+		return Positions[Index];
+	}
+	int32 GetIndex(double Time)
+	{
+		int32 Index = (int32)(((Time - Times[0]) / Step) + 0.5);
+		if (Index < 0)
+		{
+			Index = 0;
+		}
+		else if (Index >= Positions.Num())
+		{
+			Index = Positions.Num();
+		}
+		return Index;
+	}
+};
+
+// Add a key at the bisection of the two key times(indices). We use cache indices to keep things on frames.
+// Doing a bisection will result in more accurate results with less keys then doing forward tolerance checks.
+static void Bisect(FCurveModel* Curve, FCurveCache& Cache, double Tolerance, const FKeyAttributes& NewDefaultAttribute,
+	int32 StartCacheIndex, int32 EndCacheIndex, TArray<FKeyHandle>& HandlesAdded)
+{
+	if (EndCacheIndex > (StartCacheIndex +1))
+	{
+		double CurveValue = 0.0, CacheValue = 0.0;
+		const int32 EvalIndex = (int32)((double(EndCacheIndex + StartCacheIndex)) * 0.5);
+		const double EvalTime = Cache.Times[EvalIndex];
+		if (Curve->Evaluate(EvalTime, CurveValue))
+		{
+			CacheValue = Cache.Evaluate(EvalTime);
+			const double Diff = FMath::Abs(CurveValue - CacheValue);
+			if (Diff > Tolerance)
+			{
+				FKeyPosition KeyPosition(EvalTime, CacheValue);
+				TOptional<FKeyHandle> NewHandle = Curve->AddKey(KeyPosition, NewDefaultAttribute);
+				if (NewHandle.IsSet())
+				{
+					HandlesAdded.Add(NewHandle.GetValue());
+				}
+			}
+			Bisect(Curve, Cache, Tolerance, NewDefaultAttribute, StartCacheIndex, EvalIndex, HandlesAdded);
+			Bisect(Curve, Cache, Tolerance, NewDefaultAttribute, EvalIndex, EndCacheIndex, HandlesAdded);
+		}
+	}
+};
+
+//Will key reduce over the RANGE of the key handles if they are set
+//Will first find all keys that are cubic and non-auto over that set and save them as keys to keep that way we keep custom tangents,
+//unless bNeedToTestExisting is set to false since we know there will be no custom tangents, like if we just did a bake.
+//We then find keys that change velocities(e.g. they are peaks or valleys),
+//Then finally we do bisections over those intervals to finiush up the key reduction
+void UCurveEditorSmartReduceFilter::SmartReduce(FCurveModel* Curve, const FSmartReduceParams& InParams, const TOptional<FKeyHandleSet>& KeyHandleSet,
+	const bool bNeedToTestExisting, FKeyHandleSet& OutHandleSet)
+{
+	TArray<FKeyHandle> KeyHandles;
+	TArray<FKeyPosition> SelectedKeyPositions;
+	double MinKey = TNumericLimits<double>::Max(), MaxKey = TNumericLimits<double>::Lowest();
+
+	//if key's are set use that to find range to reduce otherwise use full range.
+	if (KeyHandleSet.IsSet())
+	{
+		KeyHandles.Reset(KeyHandleSet.GetValue().Num());
+		KeyHandles.Append(KeyHandleSet.GetValue().AsArray().GetData(), KeyHandleSet.GetValue().Num());
+		//assume in order todo check
+		SelectedKeyPositions.SetNum(KeyHandles.Num());
+		Curve->GetKeyPositions(KeyHandles, SelectedKeyPositions);
+		// Find the hull of the range of the selected keys
+		for (FKeyPosition Key : SelectedKeyPositions)
+		{
+			MinKey = FMath::Min(Key.InputValue, MinKey);
+			MaxKey = FMath::Max(Key.InputValue, MaxKey);
+		}
+	}
+	else //get all keys
+	{
+		Curve->GetTimeRange(MinKey, MaxKey);
+	}
+
+	// Get all keys that exist between the time range
+	KeyHandles.Reset();
+	Curve->GetKeys(MinKey, MaxKey, TNumericLimits<double>::Lowest(), TNumericLimits<double>::Max(), KeyHandles);
+	if (KeyHandles.Num() <= 2)
+	{
+		return;
+	}
+	
+	TMap<FFrameNumber, TPair<FKeyPosition, FKeyAttributes>> KeysToSet;
+	TArray<FKeyPosition> NewPositions;
+	TArray<FKeyAttributes> NewAttributes;
+	FKeyPosition KeyPosition;
+	bool bKeptFirstKey = false, bKeptLastKey = false;
+	TArray<double> KeyTimes;
+	TArray<TPair<FKeyPosition, FKeyAttributes>> KeyPosAttrs;
+	if (bNeedToTestExisting)
+	{
+		auto IsUserSpecifiedTangentKey = [](const FKeyAttributes& KeyAttributes)
+		{
+			return KeyAttributes.HasTangentMode() &&
+				(KeyAttributes.GetInterpMode() == RCIM_Cubic &&
+					(KeyAttributes.GetTangentMode() != RCTM_Auto && KeyAttributes.GetTangentMode() != RCTM_SmartAuto));
+		};
+		TSortedMap<double, TPair<FKeyPosition, FKeyAttributes>> KeysToKeep;
+		TArray<FKeyPosition> KeyPositions;
+		KeyPositions.SetNum(KeyHandles.Num());
+		Curve->GetKeyPositions(KeyHandles, KeyPositions);
+		TArray<FKeyAttributes> KeyAttributes;
+		KeyAttributes.SetNum(KeyHandles.Num());
+		Curve->GetKeyAttributes(KeyHandles, KeyAttributes);
+
+		for (int32 Index = 0; Index < KeyAttributes.Num(); ++Index)
+		{
+			if (IsUserSpecifiedTangentKey(KeyAttributes[Index]))
+			{
+				if (Index == 0)
+				{
+					bKeptFirstKey = true;
+				}
+				else if (Index == KeyAttributes.Num() - 1)
+				{
+					bKeptLastKey = true;
+				}
+				TPair<FKeyPosition, FKeyAttributes> PosAttribute(KeyPositions[Index], KeyAttributes[Index]);
+				KeysToKeep.Add(KeyPositions[Index].InputValue, PosAttribute);
+			}
+		}
+		KeysToKeep.GenerateKeyArray(KeyTimes);
+		KeysToKeep.GenerateValueArray(KeyPosAttrs);
+	}
+
+	FCurveCache Cache;
+	Cache.Step = InParams.SampleRate.AsInterval();
+	Cache.Min = TNumericLimits<double>::Max();
+	Cache.Max = TNumericLimits<double>::Lowest();
+	double Value = 0.0;
+
+	for (double Time = MinKey; Time <= MaxKey; Time += Cache.Step)
+	{
+		Cache.Times.Add(Time);
+		Curve->Evaluate(Time, Value);
+		Cache.Positions.Add(Value);
+		if (Value < Cache.Min)
+		{
+			Cache.Min = Value;
+		}
+		if (Value > Cache.Max)
+		{
+			Cache.Max = Value;
+		}
+	}
+	auto Difference = [](const TArray<double>& Input, TArray<double>& Output)
+	{
+		if (Input.Num() > 1)
+		{
+			Output.SetNum(Input.Num());
+			for (int32 Index = 0; Index < Input.Num() - 1; ++Index)
+			{
+				Output[Index] = Input[Index + 1] - Input[Index];
+			}
+			Output[Output.Num() - 1] = Output[Output.Num() - 1];
+		}
+	};
+
+	TArray<double> Velocities;
+	Difference(Cache.Positions, Velocities);
+	Curve->Modify();
+	Curve->RemoveKeys(KeyHandles);
+	FKeyAttributes NewDefaultAttribute;
+	NewDefaultAttribute.SetInterpMode(ERichCurveInterpMode::RCIM_Cubic);
+	NewDefaultAttribute.SetTangentMode(ERichCurveTangentMode::RCTM_SmartAuto);
+
+	if (bKeptFirstKey == false)
+	{
+		KeyPosition.InputValue = Cache.Times[0];
+		KeyPosition.OutputValue = Cache.Positions[0];
+		NewPositions.Add(KeyPosition);
+		NewAttributes.Add(NewDefaultAttribute);
+	}
+	int32 KeyIndex = 0;
+	double PrevVelocitySign = FMath::Sign(Velocities[0]);
+	double PrevPosition = Cache.Positions[0];
+	// go through all of the cached times, if we have keys that we need to keep(non auto's) we add them 
+	//otherwise we add a default auto based upon if it changes velocity sign or if it fails the difference threshold via bisection
+	for (int32 Index = 1; Index < Cache.Positions.Num()-1; ++Index)
+	{
+		bool bKeyWasAdded = false;
+		if (KeyIndex < KeyPosAttrs.Num())
+		{
+			while ((KeyTimes[KeyIndex] <= (Cache.Times[Index] + UE_DOUBLE_SMALL_NUMBER)) && KeyIndex < KeyTimes.Num())
+			{
+				KeyPosition.InputValue = KeyPosAttrs[KeyIndex].Key.InputValue;
+				KeyPosition.OutputValue = KeyPosAttrs[KeyIndex].Key.OutputValue;
+				NewPositions.Add(KeyPosition);
+				NewAttributes.Add(KeyPosAttrs[KeyIndex].Value);
+				if (FMath::IsNearlyEqual(KeyTimes[KeyIndex], Cache.Times[Index]))
+				{
+					bKeyWasAdded = true;
+				}
+				++KeyIndex;
+			}
+		}
+		
+		const double VelocitySign = FMath::Sign(Velocities[Index]);
+		if (bKeyWasAdded == false && (VelocitySign != PrevVelocitySign && FMath::IsNearlyEqual(Cache.Positions[Index],PrevPosition) == false))
+		{ 
+			KeyPosition.InputValue = Cache.Times[Index];
+			KeyPosition.OutputValue = Cache.Positions[Index];
+			NewPositions.Add(KeyPosition);
+			NewAttributes.Add(NewDefaultAttribute);
+		}
+		PrevVelocitySign = VelocitySign;
+		PrevPosition = Cache.Positions[Index];
+	}
+
+	if (bKeptLastKey == false)
+	{
+		KeyPosition.InputValue = Cache.Times[Cache.Positions.Num() - 1];
+		KeyPosition.OutputValue = Cache.Positions[Cache.Positions.Num() - 1];
+		NewPositions.Add(KeyPosition);
+		NewAttributes.Add(NewDefaultAttribute);
+	}
+	else
+	{
+		KeyIndex = KeyPosAttrs.Num() - 1;
+		KeyPosition.InputValue = KeyPosAttrs[KeyIndex].Key.InputValue;
+		KeyPosition.OutputValue = KeyPosAttrs[KeyIndex].Key.OutputValue;
+		NewPositions.Add(KeyPosition);
+		NewAttributes.Add(KeyPosAttrs[KeyIndex].Value);
+	}
+	
+	TArray<TOptional<FKeyHandle>> NewKeyHandles;
+	NewKeyHandles.SetNumUninitialized(NewPositions.Num());
+
+	TArrayView<TOptional<FKeyHandle>> NewKeyHandlesView = MakeArrayView(NewKeyHandles);
+
+	Curve->AddKeys(NewPositions, NewAttributes, &NewKeyHandlesView);
+
+	for (const TOptional<FKeyHandle>& Handle : NewKeyHandlesView)
+	{
+		if (Handle.IsSet())
+		{
+			OutHandleSet.Add(Handle.GetValue(), ECurvePointType::Key);
+		}
+	}
+	const float Tolerance = (InParams.TolerancePercentage / 100.0) * (Cache.Max - Cache.Min);
+	TArray<FKeyHandle> HandlesAdded;
+	for (int32 Index = 0; Index < (NewPositions.Num() -2); ++Index)
+	{
+		const int32 StartIndex = Cache.GetIndex(NewPositions[Index].InputValue);
+		const int32 EndIndex = Cache.GetIndex(NewPositions[Index + 1].InputValue);
+		Bisect(Curve, Cache, Tolerance, NewDefaultAttribute, StartIndex, EndIndex,HandlesAdded);
+	}
+	for (const FKeyHandle& NewHandle : HandlesAdded)
+	{
+		OutHandleSet.Add(NewHandle, ECurvePointType::Key);
+	}
+}
+
+void UCurveEditorSmartReduceFilter::ApplyFilter_Impl(TSharedRef<FCurveEditor> InCurveEditor, const TMap<FCurveModelID, FKeyHandleSet>& InKeysToOperateOn, TMap<FCurveModelID, FKeyHandleSet>& OutKeysToSelect)
+{
+	OutKeysToSelect.Reset();
+	for (const TTuple<FCurveModelID, FKeyHandleSet>& Pair : InKeysToOperateOn)
+	{
+		FCurveModel* Curve = InCurveEditor->FindCurve(Pair.Key);
+		if (!Curve)
+		{
+			continue;
+		}
+		const bool bNeedToTestExisting = true;
+		FKeyHandleSet& OutHandleSet = OutKeysToSelect.Add(Pair.Key);
+
+		TOptional<FKeyHandleSet> KeyHandleSet = Pair.Value;
+		SmartReduce(Curve, SmartReduceParams, KeyHandleSet, bNeedToTestExisting, OutHandleSet);
+	}
+}
\ No newline at end of file
diff --git a/Engine/Source/Editor/CurveEditor/Public/Filters/CurveEditorSmartReduceFilter.h b/Engine/Source/Editor/CurveEditor/Public/Filters/CurveEditorSmartReduceFilter.h
new file mode 100644
--- /dev/null
+++ b/Engine/Source/Editor/CurveEditor/Public/Filters/CurveEditorSmartReduceFilter.h
@@ -0,0 +1,62 @@
+// Copyright Epic Games, Inc. All Rights Reserved.
+
+#pragma once
+
+#include "Containers/Map.h"
+#include "Filters/CurveEditorFilterBase.h"
+#include "Misc/FrameRate.h"
+#include "Templates/SharedPointer.h"
+#include "UObject/ObjectMacros.h"
+#include "UObject/UObjectGlobals.h"
+
+#include "CurveEditorSmartReduceFilter.generated.h"
+
+class FCurveEditor;
+class UObject;
+struct FCurveModelID;
+struct FKeyHandleSet;
+
+
+USTRUCT(BlueprintType)
+struct FSmartReduceParams
+{
+	GENERATED_BODY()
+	/** Tolerance threshold, set as a percentage of the value's range  */
+	UPROPERTY(EditAnywhere, BlueprintReadWrite, Category = "Settings", meta = (ClampMin = "0.0", ClampMax = "100.0", UIMin = "0.0", UIMax = "100.0"))
+	float TolerancePercentage = 5.0;
+
+	/** Rate at which the curve should be sampled to compare against value tolerance. */
+	UPROPERTY(EditAnywhere, BlueprintReadWrite, Category = "Settings")
+	FFrameRate SampleRate = FFrameRate(30, 1);
+
+	void Reset()
+	{
+		TolerancePercentage = 5.0;
+		SampleRate = FFrameRate(30, 1);
+	}
+};
+
+UCLASS(DisplayName = "Smart Reduce", CollapseCategories)
+class UCurveEditorSmartReduceFilter : public UCurveEditorFilterBase
+{
+	GENERATED_BODY()
+public:
+	UCurveEditorSmartReduceFilter() {};
+	/**
+	* @param Curve to perform a reduction that works by finding where values change velocity and then perform a tolerance based bisection on these intervals.
+	* @param InParams Params to use to peform reduction
+	* @param KeyHandleSet Optional set of keys to operate on, if not set we will operate on all of them
+	* @param bNeedToTestExisting If true we will test existing keys for custom tangents so we keep them, if false we don't need to test, for example after a bake
+	* @param OutHandleSet The set of new keys.
+	*/
+	static CURVEEDITOR_API void SmartReduce(FCurveModel* Curve, const FSmartReduceParams& InParams, const TOptional<FKeyHandleSet>& KeyHandleSet,
+		const bool bNeedToTestExisting, FKeyHandleSet& OutHandleSet);
+protected:
+	void ApplyFilter_Impl(TSharedRef<FCurveEditor> InCurveEditor, const TMap<FCurveModelID, FKeyHandleSet>& InKeysToOperateOn, TMap<FCurveModelID, FKeyHandleSet>& OutKeysToSelect) override;
+
+public:
+
+	UPROPERTY(EditAnywhere, BlueprintReadWrite, Category = "Settings")
+	FSmartReduceParams SmartReduceParams;
+
+};
\ No newline at end of file
