Curve Editor Filter: Static Analysis fix

diff --git a/Engine/Source/Editor/CurveEditor/Private/Filters/CurveEditorSmartReduceFilter.cpp b/Engine/Source/Editor/CurveEditor/Private/Filters/CurveEditorSmartReduceFilter.cpp
--- a/Engine/Source/Editor/CurveEditor/Private/Filters/CurveEditorSmartReduceFilter.cpp
+++ b/Engine/Source/Editor/CurveEditor/Private/Filters/CurveEditorSmartReduceFilter.cpp
@@ -88,6 +88,26 @@ static void Bisect(FCurveModel* Curve, FCurveCache& Cache, double Tolerance, con
 	}
 };
 
+static bool IsUserSpecifiedTangentKey(const FKeyAttributes& KeyAttributes)
+{
+	return KeyAttributes.HasTangentMode() &&
+		(KeyAttributes.GetInterpMode() == RCIM_Cubic &&
+			(KeyAttributes.GetTangentMode() != RCTM_Auto && KeyAttributes.GetTangentMode() != RCTM_SmartAuto));
+};
+
+static void Difference(const TArray<double>& Input, TArray<double>& Output)
+{
+	if (Input.Num() > 1)
+	{
+		Output.SetNum(Input.Num());
+		for (int32 Index = 0; Index < Input.Num() - 1; ++Index)
+		{
+			Output[Index] = Input[Index + 1] - Input[Index];
+		}
+		Output[Output.Num() - 1] = Output[Output.Num() - 2];
+	}
+
+}
 //Will key reduce over the RANGE of the key handles if they are set
 //Will first find all keys that are cubic and non-auto over that set and save them as keys to keep that way we keep custom tangents,
 //unless bNeedToTestExisting is set to false since we know there will be no custom tangents, like if we just did a bake.
@@ -96,6 +116,7 @@ static void Bisect(FCurveModel* Curve, FCurveCache& Cache, double Tolerance, con
 void UCurveEditorSmartReduceFilter::SmartReduce(FCurveModel* Curve, const FSmartReduceParams& InParams, const TOptional<FKeyHandleSet>& KeyHandleSet,
 	const bool bNeedToTestExisting, FKeyHandleSet& OutHandleSet)
 {
+
 	TArray<FKeyHandle> KeyHandles;
 	TArray<FKeyPosition> SelectedKeyPositions;
 	double MinKey = TNumericLimits<double>::Max(), MaxKey = TNumericLimits<double>::Lowest();
@@ -137,12 +158,6 @@ void UCurveEditorSmartReduceFilter::SmartReduce(FCurveModel* Curve, const FSmart
 	TArray<TPair<FKeyPosition, FKeyAttributes>> KeyPosAttrs;
 	if (bNeedToTestExisting)
 	{
-		auto IsUserSpecifiedTangentKey = [](const FKeyAttributes& KeyAttributes)
-		{
-			return KeyAttributes.HasTangentMode() &&
-				(KeyAttributes.GetInterpMode() == RCIM_Cubic &&
-					(KeyAttributes.GetTangentMode() != RCTM_Auto && KeyAttributes.GetTangentMode() != RCTM_SmartAuto));
-		};
 		TSortedMap<double, TPair<FKeyPosition, FKeyAttributes>> KeysToKeep;
 		TArray<FKeyPosition> KeyPositions;
 		KeyPositions.SetNum(KeyHandles.Num());
@@ -191,18 +206,6 @@ void UCurveEditorSmartReduceFilter::SmartReduce(FCurveModel* Curve, const FSmart
 			Cache.Max = Value;
 		}
 	}
-	auto Difference = [](const TArray<double>& Input, TArray<double>& Output)
-	{
-		if (Input.Num() > 1)
-		{
-			Output.SetNum(Input.Num());
-			for (int32 Index = 0; Index < Input.Num() - 1; ++Index)
-			{
-				Output[Index] = Input[Index + 1] - Input[Index];
-			}
-			Output[Output.Num() - 1] = Output[Output.Num() - 1];
-		}
-	};
 
 	TArray<double> Velocities;
 	Difference(Cache.Positions, Velocities);
