Curve Editor: Invert selection (Ctrl-I)

Invert selection of keys for curves that have selected keys. (ie. don't select all keys if there are no keys selected for a curve)

diff --git a/Engine/Source/Editor/CurveEditor/Private/CurveEditor.cpp b/Engine/Source/Editor/CurveEditor/Private/CurveEditor.cpp
--- a/Engine/Source/Editor/CurveEditor/Private/CurveEditor.cpp
+++ b/Engine/Source/Editor/CurveEditor/Private/CurveEditor.cpp
@@ -432,6 +432,7 @@ void FCurveEditor::BindCommands()
 	CommandList->MapAction(FCurveEditorCommands::Get().SelectForward, FExecuteAction::CreateSP(this, &FCurveEditor::SelectForward));
 	CommandList->MapAction(FCurveEditorCommands::Get().SelectBackward, FExecuteAction::CreateSP(this, &FCurveEditor::SelectBackward));
 	CommandList->MapAction(FCurveEditorCommands::Get().SelectNone, FExecuteAction::CreateSP(this, &FCurveEditor::SelectNone));
+	CommandList->MapAction(FCurveEditorCommands::Get().InvertSelection, FExecuteAction::CreateSP(this, &FCurveEditor::InvertSelection));
 
 	{
 		FExecuteAction   ToggleInputSnapping     = FExecuteAction::CreateSP(this,   &FCurveEditor::ToggleInputSnapping);
@@ -1161,6 +1162,27 @@ void FCurveEditor::SelectNone()
 	Selection.Clear();
 }
 
+void FCurveEditor::InvertSelection()
+{
+	for (const TTuple<FCurveModelID, FKeyHandleSet>& Pair : Selection.GetAll())
+	{
+		if (FCurveModel* Curve = FindCurve(Pair.Key))
+		{
+			TArray<FKeyHandle> KeyHandles;
+			Curve->GetKeys(TNumericLimits<double>::Lowest(), TNumericLimits<double>::Max(), TNumericLimits<double>::Lowest(), TNumericLimits<double>::Max(), KeyHandles);
+			
+			TArrayView<const FKeyHandle> SelectedKeyHandles = Pair.Value.AsArray();
+				
+			for (const FKeyHandle& SelectedKeyHandle : SelectedKeyHandles)
+			{
+				KeyHandles.Remove(SelectedKeyHandle);
+			}
+
+			Selection.Remove(Pair.Key, ECurvePointType::Key, SelectedKeyHandles);
+			Selection.Add(Pair.Key, ECurvePointType::Key, KeyHandles);
+		}
+	}	
+}
 
 bool FCurveEditor::IsInputSnappingEnabled() const
 {
diff --git a/Engine/Source/Editor/CurveEditor/Private/CurveEditorCommands.cpp b/Engine/Source/Editor/CurveEditor/Private/CurveEditorCommands.cpp
--- a/Engine/Source/Editor/CurveEditor/Private/CurveEditorCommands.cpp
+++ b/Engine/Source/Editor/CurveEditor/Private/CurveEditorCommands.cpp
@@ -90,6 +90,7 @@ void FCurveEditorCommands::RegisterCommands()
 	UI_COMMAND(SelectBackward, "Select All Keys Backward", "Select all keys backward from the current time", EUserInterfaceActionType::Button, FInputChord(EModifierKey::Control, EKeys::LeftBracket));
 
 	UI_COMMAND(SelectNone, "Select None", "Select none", EUserInterfaceActionType::Button, FInputChord(EKeys::Escape));
+	UI_COMMAND(InvertSelection, "Invert Selection", "Invert the selection of keys on curves with selected keys.", EUserInterfaceActionType::Button, FInputChord(EModifierKey::Control, EKeys::I));
 
 	// Buffer and Apply Curves. Like copy and paste, but directly onto the curves they were stored from.
 	// These names are overwritten in CurveEditorContextMenu to show the number of stashed curves.
diff --git a/Engine/Source/Editor/CurveEditor/Public/CurveEditor.h b/Engine/Source/Editor/CurveEditor/Public/CurveEditor.h
--- a/Engine/Source/Editor/CurveEditor/Public/CurveEditor.h
+++ b/Engine/Source/Editor/CurveEditor/Public/CurveEditor.h
@@ -244,6 +244,7 @@ public:
 	void SelectForward();
 	void SelectBackward();
 	void SelectNone();
+	void InvertSelection();
 
 	/** Toggle the expansion state of the selected nodes or all nodes if none selected */
 	void ToggleExpandCollapseNodes(bool bRecursive);
diff --git a/Engine/Source/Editor/CurveEditor/Public/CurveEditorCommands.h b/Engine/Source/Editor/CurveEditor/Public/CurveEditorCommands.h
--- a/Engine/Source/Editor/CurveEditor/Public/CurveEditorCommands.h
+++ b/Engine/Source/Editor/CurveEditor/Public/CurveEditorCommands.h
@@ -97,6 +97,7 @@ public:
 	TSharedPtr<FUICommandInfo> SelectForward;
 	TSharedPtr<FUICommandInfo> SelectBackward;
 	TSharedPtr<FUICommandInfo> SelectNone;
+	TSharedPtr<FUICommandInfo> InvertSelection;
 
 	TSharedPtr<FUICommandInfo> BufferVisibleCurves;
 	TSharedPtr<FUICommandInfo> SwapBufferedCurves;
