Curve Editor: Change snap axis to a setting/preference

Tested that you can still override snapping by holding down Shift.
Tested that snapping settings are retained when closing and reopening Sequencer/Curve Editor.

diff --git a/Engine/Plugins/VirtualProduction/CameraCalibration/Source/CameraCalibrationEditor/Private/AssetEditor/SLensDataViewer.cpp b/Engine/Plugins/VirtualProduction/CameraCalibration/Source/CameraCalibrationEditor/Private/AssetEditor/SLensDataViewer.cpp
--- a/Engine/Plugins/VirtualProduction/CameraCalibration/Source/CameraCalibrationEditor/Private/AssetEditor/SLensDataViewer.cpp
+++ b/Engine/Plugins/VirtualProduction/CameraCalibration/Source/CameraCalibrationEditor/Private/AssetEditor/SLensDataViewer.cpp
@@ -142,7 +142,7 @@ void SLensDataViewer::Construct(const FArguments& InArgs, ULensFile* InLensFile,
 
 	// Snap only Y axis
 	FCurveEditorAxisSnap SnapYAxisOnly = CurveEditor->GetAxisSnap();
-	SnapYAxisOnly.RestrictedAxisList = EAxisList::Type::Y;
+	SnapYAxisOnly.RestrictedAxisList = ECurveEditorSnapAxis::CESA_Y;
 	CurveEditor->SetAxisSnap(SnapYAxisOnly);
 
 	CurvePanel = SNew(SCameraCalibrationCurveEditorPanel, CurveEditor.ToSharedRef(), TimeSliderController);
diff --git a/Engine/Source/Editor/CurveEditor/Private/CurveEditor.cpp b/Engine/Source/Editor/CurveEditor/Private/CurveEditor.cpp
--- a/Engine/Source/Editor/CurveEditor/Private/CurveEditor.cpp
+++ b/Engine/Source/Editor/CurveEditor/Private/CurveEditor.cpp
@@ -85,6 +85,7 @@ FCurveEditor::FCurveEditor()
 	GridLineLabelFormatYAttribute = LOCTEXT("GridYLabelFormat", "{0}");
 	
 	Settings->GetOnCustomColorsChanged().AddRaw(this, &FCurveEditor::OnCustomColorsChanged);
+	Settings->GetOnAxisSnappingChanged().AddRaw(this, &FCurveEditor::OnAxisSnappingChanged);
 }
 
 FCurveEditor::~FCurveEditor()
@@ -92,6 +93,7 @@ FCurveEditor::~FCurveEditor()
 	if (Settings)
 	{
 		Settings->GetOnCustomColorsChanged().RemoveAll(this);
+		Settings->GetOnAxisSnappingChanged().RemoveAll(this);
 	}
 }
 
@@ -2379,4 +2381,13 @@ void FCurveEditor::OnCustomColorsChanged()
 	}
 }
 
+void FCurveEditor::OnAxisSnappingChanged()
+{
+	TSharedPtr<SCurveEditorPanel> Panel = WeakPanel.Pin();
+	if (Panel.IsValid())
+	{
+		Panel->UpdateAxisSnapping();
+	}
+}
+
 #undef LOCTEXT_NAMESPACE
diff --git a/Engine/Source/Editor/CurveEditor/Private/CurveEditorSettings.cpp b/Engine/Source/Editor/CurveEditor/Private/CurveEditorSettings.cpp
--- a/Engine/Source/Editor/CurveEditor/Private/CurveEditorSettings.cpp
+++ b/Engine/Source/Editor/CurveEditor/Private/CurveEditorSettings.cpp
@@ -14,6 +14,7 @@ UCurveEditorSettings::UCurveEditorSettings()
 	bShowCurveEditorCurveToolTips = true;
 	TangentVisibility = ECurveEditorTangentVisibility::SelectedKeys;
 	ZoomPosition = ECurveEditorZoomPosition::CurrentTime;
+	SnapAxis = ECurveEditorSnapAxis::CESA_None;
 	bSnapTimeToSelection = false;
 
 	SelectionColor = FLinearColor::White;
@@ -135,6 +136,20 @@ void UCurveEditorSettings::SetZoomPosition(ECurveEditorZoomPosition InZoomPositi
 	}
 }
 
+ECurveEditorSnapAxis UCurveEditorSettings::GetSnapAxis() const
+{
+	return SnapAxis;
+}
+
+void UCurveEditorSettings::SetSnapAxis(ECurveEditorSnapAxis InSnapAxis)
+{
+	if (SnapAxis != InSnapAxis)
+	{
+		SnapAxis = InSnapAxis;
+		SaveConfig();
+	}
+}
+
 bool UCurveEditorSettings::GetSnapTimeToSelection() const
 {
 	return bSnapTimeToSelection;
@@ -330,6 +345,11 @@ void UCurveEditorSettings::PostEditChangeProperty(FPropertyChangedEvent& Propert
 	{
 		OnCustomColorsChangedEvent.Broadcast();
 	}
-	
+	else if (PropertyName == GET_MEMBER_NAME_CHECKED(UCurveEditorSettings, SnapAxis) ||
+		MemberPropertyName == GET_MEMBER_NAME_CHECKED(UCurveEditorSettings, SnapAxis))
+	{
+		OnAxisSnappingChangedEvent.Broadcast();
+	}
+
 	Super::PostEditChangeProperty(PropertyChangedEvent);
 }
diff --git a/Engine/Source/Editor/CurveEditor/Private/SCurveEditorPanel.cpp b/Engine/Source/Editor/CurveEditor/Private/SCurveEditorPanel.cpp
--- a/Engine/Source/Editor/CurveEditor/Private/SCurveEditorPanel.cpp
+++ b/Engine/Source/Editor/CurveEditor/Private/SCurveEditorPanel.cpp
@@ -313,6 +313,8 @@ void SCurveEditorPanel::Construct(const FArguments& InArgs, TSharedRef<FCurveEdi
 
 	UpdateEditBox();
 
+	UpdateAxisSnapping();
+
 	// Initializes our Curve Views on the next Tick
 	SetViewMode(ECurveEditorViewID::Absolute);
 }
@@ -460,9 +462,9 @@ void SCurveEditorPanel::BindCommands()
 
 	// Axis Snapping
 	{
-		FUIAction SetSnappingNoneAction(FExecuteAction::CreateSP(this, &SCurveEditorPanel::SetAxisSnapping, EAxisList::Type::None));
-		FUIAction SetSnappingHorizontalAction(FExecuteAction::CreateSP(this, &SCurveEditorPanel::SetAxisSnapping, EAxisList::Type::X));
-		FUIAction SetSnappingVerticalAction(FExecuteAction::CreateSP(this, &SCurveEditorPanel::SetAxisSnapping, EAxisList::Type::Y));
+		FUIAction SetSnappingNoneAction(FExecuteAction::CreateSP(this, &SCurveEditorPanel::SetAxisSnapping, ECurveEditorSnapAxis::CESA_None));
+		FUIAction SetSnappingHorizontalAction(FExecuteAction::CreateSP(this, &SCurveEditorPanel::SetAxisSnapping, ECurveEditorSnapAxis::CESA_X));
+		FUIAction SetSnappingVerticalAction(FExecuteAction::CreateSP(this, &SCurveEditorPanel::SetAxisSnapping, ECurveEditorSnapAxis::CESA_Y));
 
 		CommandList->MapAction(FCurveEditorCommands::Get().SetAxisSnappingNone, SetSnappingNoneAction);
 		CommandList->MapAction(FCurveEditorCommands::Get().SetAxisSnappingHorizontal, SetSnappingHorizontalAction);
@@ -482,10 +484,16 @@ bool SCurveEditorPanel::CompareViewMode(const ECurveEditorViewID InViewMode) con
 	return DefaultViewID == InViewMode;
 }
 
-void SCurveEditorPanel::SetAxisSnapping(EAxisList::Type InAxis)
+void SCurveEditorPanel::SetAxisSnapping(ECurveEditorSnapAxis InAxis)
+{
+	CurveEditor->GetSettings()->SetSnapAxis(InAxis);
+	UpdateAxisSnapping();
+}
+
+void SCurveEditorPanel::UpdateAxisSnapping()
 {
 	FCurveEditorAxisSnap Snap = CurveEditor->GetAxisSnap();
-	Snap.RestrictedAxisList = InAxis;
+	Snap.RestrictedAxisList = CurveEditor->GetSettings()->GetSnapAxis();
 	CurveEditor->SetAxisSnap(Snap);
 }
 
@@ -1264,11 +1272,11 @@ TSharedPtr<FExtender> SCurveEditorPanel::GetToolbarExtender()
 				AxisSnappingModeIcon.Bind(TAttribute<FSlateIcon>::FGetter::CreateLambda([InEditorPanel] {
 					switch (InEditorPanel->GetCurveEditor()->GetAxisSnap().RestrictedAxisList)
 					{
-					case EAxisList::Type::X:
+					case ECurveEditorSnapAxis::CESA_X:
 						return FCurveEditorCommands::Get().SetAxisSnappingHorizontal->GetIcon();
-					case EAxisList::Type::Y:
+					case ECurveEditorSnapAxis::CESA_Y:
 						return FCurveEditorCommands::Get().SetAxisSnappingVertical->GetIcon();
-					default: // EKeyGroupMode::None
+					default: // ECurveEditorSnapAxis::CESA_None
 						return FCurveEditorCommands::Get().SetAxisSnappingNone->GetIcon();
 					}
 				}));
diff --git a/Engine/Source/Editor/CurveEditor/Public/CurveEditor.h b/Engine/Source/Editor/CurveEditor/Public/CurveEditor.h
--- a/Engine/Source/Editor/CurveEditor/Public/CurveEditor.h
+++ b/Engine/Source/Editor/CurveEditor/Public/CurveEditor.h
@@ -596,6 +596,7 @@ protected:
 	void ApplyBufferedCurveToTarget(const IBufferedCurveModel* BufferedCurve, FCurveModel* TargetCurve);
 
 	void OnCustomColorsChanged();
+	void OnAxisSnappingChanged();
 
 protected:
 
diff --git a/Engine/Source/Editor/CurveEditor/Public/CurveEditorSettings.h b/Engine/Source/Editor/CurveEditor/Public/CurveEditorSettings.h
--- a/Engine/Source/Editor/CurveEditor/Public/CurveEditorSettings.h
+++ b/Engine/Source/Editor/CurveEditor/Public/CurveEditorSettings.h
@@ -40,6 +40,18 @@ enum class ECurveEditorZoomPosition : uint8
 	MousePosition UMETA(DisplayName = "Mouse Position"),
 };
 
+/** Defines the axis to snap to when dragging. */
+UENUM()
+enum class ECurveEditorSnapAxis : uint8
+{
+	/** Don't snap to any axis when dragging. */
+	CESA_None UMETA(DisplayName = "None"),
+	/* Snap to the x axis when dragging. */
+	CESA_X UMETA(DisplayName = "X Only"),
+	/* Snap to the y axis when dragging. */
+	CESA_Y UMETA(DisplayName = "Y Only")
+};
+
 /** Custom Color Object*/
 USTRUCT()
 struct FCustomColorForChannel
@@ -81,6 +93,7 @@ public:
 	GENERATED_BODY()
 
 	DECLARE_MULTICAST_DELEGATE(FOnCustomColorsChanged);
+	DECLARE_MULTICAST_DELEGATE(FOnAxisSnappingChanged);
 
 	UCurveEditorSettings();
 
@@ -124,6 +137,11 @@ public:
 	/** Set zoom in/out position (mouse position or current time). */
 	void SetZoomPosition(ECurveEditorZoomPosition InZoomPosition);
 
+	/** Get snap axis. */
+	ECurveEditorSnapAxis GetSnapAxis() const;
+	/** Set snap axis. */
+	void SetSnapAxis(ECurveEditorSnapAxis InSnapAxis);
+
 	/** Get whether to snap the time to the currently selected key. */
 	bool GetSnapTimeToSelection() const;
 	/** Set whether to snap the time to the currently selected key. */
@@ -144,6 +162,9 @@ public:
 	/** Gets the multicast delegate which is run whenever custom colors have changed. */
 	FOnCustomColorsChanged& GetOnCustomColorsChanged() { return OnCustomColorsChangedEvent; }
 
+	/** Gets the multicast delegate which is run whenever axis snapping has changed. */
+	FOnAxisSnappingChanged& GetOnAxisSnappingChanged() { return OnAxisSnappingChangedEvent; }
+
 	/** Get custom color for space name. Parent and World are reserved names and will be used instead of the specified control name. */
 	TOptional<FLinearColor> GetSpaceSwitchColor(const FString& InControlName) const;
 	/** Set Custom Space SwitchColor for the specified control name. */
@@ -186,9 +207,12 @@ protected:
 	UPROPERTY( config, EditAnywhere, Category="Curve Editor" )
 	ECurveEditorTangentVisibility TangentVisibility;
 
-	UPROPERTY( config, EditAnywhere, Category="Curve Editor")
+	UPROPERTY(config, EditAnywhere, Category = "Curve Editor")
 	ECurveEditorZoomPosition ZoomPosition;
 
+	UPROPERTY(config, EditAnywhere, Category = "Curve Editor")
+	ECurveEditorSnapAxis SnapAxis;
+
 	UPROPERTY( config, EditAnywhere, Category="Curve Editor")
 	bool bSnapTimeToSelection;
 
@@ -212,4 +236,5 @@ protected:
 
 private:
 	FOnCustomColorsChanged OnCustomColorsChangedEvent;
+	FOnAxisSnappingChanged OnAxisSnappingChangedEvent;
 };
diff --git a/Engine/Source/Editor/CurveEditor/Public/CurveEditorSnapMetrics.h b/Engine/Source/Editor/CurveEditor/Public/CurveEditorSnapMetrics.h
--- a/Engine/Source/Editor/CurveEditor/Public/CurveEditorSnapMetrics.h
+++ b/Engine/Source/Editor/CurveEditor/Public/CurveEditorSnapMetrics.h
@@ -7,6 +7,7 @@
 #include "Misc/FrameRate.h"
 #include "Math/Axis.h"
 #include "Math/Vector2D.h"
+#include "CurveEditorSettings.h"
 
 struct FCurveSnapMetrics
 {
@@ -74,19 +75,17 @@ struct FCurveEditorAxisSnap
 	};
 
 	/** Can be set to either X, Y, or None to control which axis GetSnappedPosition snaps to. User can override None by pressing shift. */
-	EAxisList::Type RestrictedAxisList;
+	ECurveEditorSnapAxis RestrictedAxisList;
 
 	FCurveEditorAxisSnap()
 	{
-		RestrictedAxisList = EAxisList::None;
+		RestrictedAxisList = ECurveEditorSnapAxis::CESA_None;
 	}
 
 	FVector2D GetSnappedPosition(FVector2D InitialPosition, FVector2D LastPosition, FVector2D CurrentPosition, const FPointerEvent& MouseEvent, FSnapState& InOutSnapState, const bool bIgnoreAxisLock = false)
 	{
-		check(RestrictedAxisList == EAxisList::Type::None || RestrictedAxisList == EAxisList::Type::X || RestrictedAxisList == EAxisList::Type::Y);
-
 		// If we're ignoring axis lock (such as for UI) we allow them to use shift anyways
-		bool bCanUseShift = bIgnoreAxisLock || RestrictedAxisList == EAxisList::Type::None;
+		bool bCanUseShift = bIgnoreAxisLock || RestrictedAxisList == ECurveEditorSnapAxis::CESA_None;
 
 		FVector2D MouseLockVector = FVector2D::UnitVector;
 		if (bCanUseShift)
@@ -131,11 +130,11 @@ struct FCurveEditorAxisSnap
 
 			MouseLockVector = InOutSnapState.MouseLockVector;
 		}
-		else if (!bIgnoreAxisLock && RestrictedAxisList == EAxisList::Type::X)
+		else if (!bIgnoreAxisLock && RestrictedAxisList == ECurveEditorSnapAxis::CESA_X)
 		{
 			MouseLockVector.Y = 0;
 		}
-		else if (!bIgnoreAxisLock && RestrictedAxisList == EAxisList::Type::Y)
+		else if (!bIgnoreAxisLock && RestrictedAxisList == ECurveEditorSnapAxis::CESA_Y)
 		{
 			MouseLockVector.X = 0;
 		}
diff --git a/Engine/Source/Editor/CurveEditor/Public/SCurveEditorPanel.h b/Engine/Source/Editor/CurveEditor/Public/SCurveEditorPanel.h
--- a/Engine/Source/Editor/CurveEditor/Public/SCurveEditorPanel.h
+++ b/Engine/Source/Editor/CurveEditor/Public/SCurveEditorPanel.h
@@ -168,6 +168,9 @@ class CURVEEDITOR_API SCurveEditorPanel : public SCompoundWidget
 	/** Reset Stored Min/Max's*/
 	void ResetMinMaxes();
 
+	/** Update the axis snapping based on the settings. */
+	void UpdateAxisSnapping();
+
 	/** Delegate for when the chosen filter class has changed */
 	FSimpleDelegate OnFilterClassChanged;
 	void FilterClassChanged();
@@ -275,7 +278,7 @@ private:
 	bool CanSetKeyInterpolation() const;
 
 	/** Sets the axis snapping to the specified value. Only supports X, Y and None. */
-	void SetAxisSnapping(EAxisList::Type InAxis);
+	void SetAxisSnapping(ECurveEditorSnapAxis);
 
 	/** Get a reference to the curve editor this panel represents. */
 	TSharedPtr<FCurveEditor> GetCurveEditor() const { return CurveEditor; }
