Curve Editor: Tweak Invert Selection so it just removes all selection and then adds what's needed.

#rnx

diff --git a/Engine/Source/Editor/CurveEditor/Private/CurveEditor.cpp b/Engine/Source/Editor/CurveEditor/Private/CurveEditor.cpp
--- a/Engine/Source/Editor/CurveEditor/Private/CurveEditor.cpp
+++ b/Engine/Source/Editor/CurveEditor/Private/CurveEditor.cpp
@@ -1168,7 +1168,8 @@ void FCurveEditor::InvertSelection()
 {
 	for (const TTuple<FCurveModelID, FKeyHandleSet>& Pair : Selection.GetAll())
 	{
-		if (FCurveModel* Curve = FindCurve(Pair.Key))
+		FCurveModelID CurveModelID = Pair.Key;
+		if (FCurveModel* Curve = FindCurve(CurveModelID))
 		{
 			TArray<FKeyHandle> KeyHandles;
 			Curve->GetKeys(TNumericLimits<double>::Lowest(), TNumericLimits<double>::Max(), TNumericLimits<double>::Lowest(), TNumericLimits<double>::Max(), KeyHandles);
@@ -1180,8 +1181,8 @@ void FCurveEditor::InvertSelection()
 				KeyHandles.Remove(SelectedKeyHandle);
 			}
 
-			Selection.Remove(Pair.Key, ECurvePointType::Key, SelectedKeyHandles);
-			Selection.Add(Pair.Key, ECurvePointType::Key, KeyHandles);
+			Selection.Remove(CurveModelID);
+			Selection.Add(CurveModelID, ECurvePointType::Key, KeyHandles);
 		}
 	}	
 }
