Curve Editor: Tweak Invert Selection so that it skips curves that don't have keys selected

#rnx

diff --git a/Engine/Source/Editor/CurveEditor/Private/CurveEditor.cpp b/Engine/Source/Editor/CurveEditor/Private/CurveEditor.cpp
--- a/Engine/Source/Editor/CurveEditor/Private/CurveEditor.cpp
+++ b/Engine/Source/Editor/CurveEditor/Private/CurveEditor.cpp
@@ -1176,13 +1176,16 @@ void FCurveEditor::InvertSelection()
 			
 			TArrayView<const FKeyHandle> SelectedKeyHandles = Pair.Value.AsArray();
 				
-			for (const FKeyHandle& SelectedKeyHandle : SelectedKeyHandles)
+			if (SelectedKeyHandles.Num() > 0)
 			{
-				KeyHandles.Remove(SelectedKeyHandle);
-			}
+				for (const FKeyHandle& SelectedKeyHandle : SelectedKeyHandles)
+				{
+					KeyHandles.Remove(SelectedKeyHandle);
+				}
 
-			Selection.Remove(CurveModelID);
-			Selection.Add(CurveModelID, ECurvePointType::Key, KeyHandles);
+				Selection.Remove(CurveModelID);
+				Selection.Add(CurveModelID, ECurvePointType::Key, KeyHandles);
+			}
 		}
 	}	
 }
