Cureve Editor: Prevent setting of invalid output bounds

diff --git a/Engine/Source/Editor/CurveEditor/Private/SCurveEditorView.cpp b/Engine/Source/Editor/CurveEditor/Private/SCurveEditorView.cpp
--- a/Engine/Source/Editor/CurveEditor/Private/SCurveEditorView.cpp
+++ b/Engine/Source/Editor/CurveEditor/Private/SCurveEditorView.cpp
@@ -116,7 +116,7 @@ void SCurveEditorView::FrameVertical(double InOutputMin, double InOutputMax)
 
 void SCurveEditorView::SetOutputBounds(double InOutputMin, double InOutputMax)
 {
-	if (!bFixedOutputBounds)
+	if (!bFixedOutputBounds && InOutputMin < InOutputMax)
 	{
 		OutputMin = InOutputMin;
 		OutputMax = InOutputMax;
