Sequencer: Added time-warp capabilities to sequences, sub-sections and skeletal animation sections

This is a large suite of changes that are being submitted in bulk for build-health reasons:

Core machinery:
 - FMovieSceneNumericVariant: This is the base variant type that represents either a double, or a custom object. It is not much use on its own, but provides the basis for the time-warp variant, and will be the basis for parameterization.
 - FMovieSceneTimeWarpVariant: This is a drop in replacement for a double PlayRate/TimeScale UProperty that can now represent a time-warp or play rate curve, a loop, a clamp, or a fixed time. It also provides functionality for remapping time, and inverse remapping time.
 - The following interpolation types have been added/improved to add Derivative and Integral functionality
   * FConstantValue: Represents a curve of the form y=c
   * FLinearInterpolation: Represents a curve of the form y=mx + c
   * FQuadraticInterpolation: Represents a curve of the form y=ax^2 + bx + c.
   * FCubicInterpolation: Represents a curve of the form y=ax^3 + bx^2 + cx + d.
   * FQuarticInterpolation: Represents a curve of the form y=ax^4 + bx^3 + cx^2 + dx + e. Does not support integration.
   * FCubicBezierInterpolation now has an AsCubic method that converts it to a cubic polynomial form, which also allows derivative/integral functionality.
 - FPiecewiseCurve(Model): A generic read-only piecewise curve implementation and curve editor model. Used to display derivative curves for playrate.
 - FMovieSceneTimingParameters: A struct for specifying section timing that simplifies computation of a transform. I wanted to replace all implementations with this but current deprecation mechanisms were prohibitively restrictive. More work will be needed there to remove legacy code.
 - FMovieSceneTimeWarpChannel: A new channel type that defines time in one of two domains: PlayRate or Time

Curve Editor: Added the ability to specify custom axes and child curves
- Custom Axes:
  * Custom axes can now be assigned to each curve within a view, and the axis defines the scaling for that curve.
  * This allows us to have multiple different 'spaces' on the same view that have vastly different scales (for instance, play rate, and tick resolution times)
- Child Curves:
  * Child curves are added whenever their parent is added. This allows us to show the derivative of a play rate curve on the curve editor.

Sequencer: Added Owning Object and key offsets to channel meta-data and curve models, made Sequencer key-editors polymorphic,
- New options on channel proxy entries allow us to specify that a channel's keys are relative to their section start time via a key offset
- New options in the curve editor models allow for a non-UMovieSceneSection owner object type

diff --git a/Engine/Plugins/Animation/ControlRig/Source/ControlRigEditor/Private/ControlRigSequencerEditorLibrary.cpp b/Engine/Plugins/Animation/ControlRig/Source/ControlRigEditor/Private/ControlRigSequencerEditorLibrary.cpp
--- a/Engine/Plugins/Animation/ControlRig/Source/ControlRigEditor/Private/ControlRigSequencerEditorLibrary.cpp
+++ b/Engine/Plugins/Animation/ControlRig/Source/ControlRigEditor/Private/ControlRigSequencerEditorLibrary.cpp
@@ -1379,7 +1379,7 @@ static bool LocalGetControlRigControlValues(IMovieScenePlayer* Player, UMovieSce
 			{
 				GlobalTime = FFrameRate::TransformTime(GlobalTime, MovieScene->GetDisplayRate(), MovieScene->GetTickResolution());
 			}
-			GlobalTime = GlobalTime * RootToLocalTransform.InverseNoLooping();
+			GlobalTime = RootToLocalTransform.Inverse().TryTransformTime(GlobalTime).Get(GlobalTime);
 			FMovieSceneContext Context = FMovieSceneContext(FMovieSceneEvaluationRange(GlobalTime, TickResolution), Player->GetPlaybackStatus()).SetHasJumped(true);
 
 			Player->GetEvaluationTemplate().EvaluateSynchronousBlocking(Context);
diff --git a/Engine/Plugins/Animation/ControlRig/Source/ControlRigEditor/Private/Sequencer/ControlRigParameterTrackEditor.cpp b/Engine/Plugins/Animation/ControlRig/Source/ControlRigEditor/Private/Sequencer/ControlRigParameterTrackEditor.cpp
--- a/Engine/Plugins/Animation/ControlRig/Source/ControlRigEditor/Private/Sequencer/ControlRigParameterTrackEditor.cpp
+++ b/Engine/Plugins/Animation/ControlRig/Source/ControlRigEditor/Private/Sequencer/ControlRigParameterTrackEditor.cpp
@@ -28,6 +28,7 @@
 #include "SequencerSectionPainter.h"
 #include "Components/SkeletalMeshComponent.h"
 #include "Editor.h"
+#include "ISequencerChannelInterface.h"
 #include "Editor/UnrealEdEngine.h"
 #include "UnrealEdGlobals.h"
 #include "ClassViewerModule.h"
@@ -977,6 +978,8 @@ FReply SBakeToAnimAndControlRigOptionsWindow::OnResetToDefaultClick() const
 
 void FControlRigParameterTrackEditor::SmartReduce(TSharedPtr<ISequencer>& InSequencer, const FSmartReduceParams& InParams, UMovieSceneControlRigParameterSection * ParamSection)
 {
+	using namespace UE::Sequencer;
+
 	if (ParamSection)
 	{
 		ISequencerModule& SequencerModule = FModuleManager::LoadModuleChecked<ISequencerModule>("Sequencer");
@@ -992,7 +995,9 @@ void FControlRigParameterTrackEditor::SmartReduce(TSharedPtr<ISequencer>& InSequ
 			{
 				FMovieSceneChannelHandle ChannelHandle = ChannelProxy.MakeHandle(ChannelTypeName, Index);
 				ISequencerChannelInterface* EditorInterface = SequencerModule.FindChannelEditorInterface(ChannelHandle.GetChannelTypeName());
-				if (TUniquePtr<FCurveModel> CurveModel = EditorInterface->CreateCurveEditorModel_Raw(ChannelHandle, ParamSection, InSequencer.ToSharedRef()))
+
+				FCreateCurveEditorModelParams CurveModelParams{ ParamSection, ParamSection, InSequencer.ToSharedRef() };
+				if (TUniquePtr<FCurveModel> CurveModel = EditorInterface->CreateCurveEditorModel_Raw(ChannelHandle, CurveModelParams))
 				{
 					FKeyHandleSet OutHandleSet;
 					UCurveEditorSmartReduceFilter::SmartReduce(CurveModel.Get(), InParams, KeyHandleSet, bNeedToTestExisting, OutHandleSet);
@@ -4813,13 +4818,14 @@ bool FControlRigParameterTrackEditor::CollapseAllLayers(TSharedPtr<ISequencer>&S
 						ControlLocalTransforms.Add(NameTransforms);
 					}
 
+					FMovieSceneInverseSequenceTransform LocalToRootTransform = RootToLocalTransform.Inverse();
+
 					//get all of the local 
 					int32 Index = 0;
 					for (Index = 0; Index < Frames.Num(); ++Index)
 					{
 						const FFrameNumber& FrameNumber = Frames[Index];
-						FFrameTime GlobalTime(FrameNumber);
-						GlobalTime = GlobalTime * RootToLocalTransform.InverseNoLooping();
+						FFrameTime GlobalTime = LocalToRootTransform.TryTransformTime(FrameNumber).Get(FrameNumber);
 
 						FMovieSceneContext Context = FMovieSceneContext(FMovieSceneEvaluationRange(GlobalTime, TickResolution), SequencerPtr->GetPlaybackStatus()).SetHasJumped(true);
 
diff --git a/Engine/Plugins/Animation/ControlRig/Source/ControlRigEditor/Private/Sequencer/ControlRigSpaceChannelEditors.cpp b/Engine/Plugins/Animation/ControlRig/Source/ControlRigEditor/Private/Sequencer/ControlRigSpaceChannelEditors.cpp
--- a/Engine/Plugins/Animation/ControlRig/Source/ControlRigEditor/Private/Sequencer/ControlRigSpaceChannelEditors.cpp
+++ b/Engine/Plugins/Animation/ControlRig/Source/ControlRigEditor/Private/Sequencer/ControlRigSpaceChannelEditors.cpp
@@ -171,7 +171,7 @@ bool CanCreateKeyEditor(const FMovieSceneControlRigSpaceChannel* Channel)
 {
 	return false; //mz todoo maybe change
 }
-TSharedRef<SWidget> CreateKeyEditor(const TMovieSceneChannelHandle<FMovieSceneControlRigSpaceChannel>& Channel, UMovieSceneSection* Section, const FGuid& InObjectBindingID, TWeakPtr<FTrackInstancePropertyBindings> PropertyBindings, TWeakPtr<ISequencer> InSequencer)
+TSharedRef<SWidget> CreateKeyEditor(const TMovieSceneChannelHandle<FMovieSceneControlRigSpaceChannel>& Channel, const UE::Sequencer::FCreateKeyEditorParams& Params)
 {
 	return SNullWidget::NullWidget;
 }
@@ -339,7 +339,7 @@ FKeyHandle FControlRigSpaceChannelHelpers::SequencerKeyControlRigSpaceChannel(UC
 
 		//make sure to evaluate first frame
 		FFrameTime CurrentTime(Time);
-		CurrentTime = CurrentTime * RootToLocalTransform.InverseNoLooping();
+		CurrentTime = RootToLocalTransform.Inverse().TryTransformTime(CurrentTime).Get(CurrentTime);
 
 		FMovieSceneContext SceneContext = FMovieSceneContext(FMovieSceneEvaluationRange(CurrentTime, TickResolution), Sequencer->GetPlaybackStatus()).SetHasJumped(true);
 		Sequencer->GetEvaluationTemplate().EvaluateSynchronousBlocking(SceneContext);
@@ -447,6 +447,8 @@ FKeyHandle FControlRigSpaceChannelHelpers::SequencerKeyControlRigSpaceChannel(UC
 			}
 		}
 
+		FMovieSceneInverseSequenceTransform LocalToRootTransform = RootToLocalTransform.Inverse();
+
 		//do any compensation or previous key adding
 		FRigControlModifiedContext Context;
 		Context.SetKey = EControlRigSetKey::Always;
@@ -455,7 +457,7 @@ FKeyHandle FControlRigSpaceChannelHelpers::SequencerKeyControlRigSpaceChannel(UC
 		if (bSetPreviousKey)
 		{
 			FFrameTime GlobalTime(Time - 1);
-			GlobalTime = GlobalTime * RootToLocalTransform.InverseNoLooping();
+			GlobalTime = LocalToRootTransform.TryTransformTime(GlobalTime).Get(GlobalTime);
 
 			SceneContext = FMovieSceneContext(FMovieSceneEvaluationRange(GlobalTime, TickResolution), Sequencer->GetPlaybackStatus()).SetHasJumped(true);
 			Sequencer->GetEvaluationTemplate().EvaluateSynchronousBlocking(SceneContext);
@@ -489,7 +491,7 @@ FKeyHandle FControlRigSpaceChannelHelpers::SequencerKeyControlRigSpaceChannel(UC
 		for (const FFrameNumber& Frame : Frames)
 		{
 			FFrameTime GlobalTime(Frame);
-			GlobalTime = GlobalTime * RootToLocalTransform.InverseNoLooping();
+			GlobalTime = LocalToRootTransform.TryTransformTime(GlobalTime).Get(GlobalTime);
 
 			SceneContext = FMovieSceneContext(FMovieSceneEvaluationRange(GlobalTime, TickResolution), Sequencer->GetPlaybackStatus()).SetHasJumped(true);
 			Sequencer->GetEvaluationTemplate().EvaluateSynchronousBlocking(SceneContext);
@@ -699,12 +701,13 @@ void  FControlRigSpaceChannelHelpers::SequencerSpaceChannelKeyDeleted(UControlRi
 		Context.SetKey = EControlRigSetKey::Always;
 		FFrameRate TickResolution = Sequencer->GetFocusedTickResolution();
 		FMovieSceneSequenceTransform RootToLocalTransform = Sequencer->GetFocusedMovieSceneSequenceTransform();
+		FMovieSceneInverseSequenceTransform LocalToRootTransform = RootToLocalTransform.Inverse();
 
 		for (const FFrameNumber& Frame : Frames)
 		{
 			//evaluate sequencer
 			FFrameTime GlobalTime(Frame);
-			GlobalTime = GlobalTime * RootToLocalTransform.InverseNoLooping();
+			GlobalTime = LocalToRootTransform.TryTransformTime(GlobalTime).Get(GlobalTime);
 
 			FMovieSceneContext SceneContext = FMovieSceneContext(FMovieSceneEvaluationRange(GlobalTime, TickResolution), Sequencer->GetPlaybackStatus()).SetHasJumped(true);
 			Sequencer->GetEvaluationTemplate().EvaluateSynchronousBlocking(SceneContext);
@@ -1094,6 +1097,7 @@ void FControlRigSpaceChannelHelpers::SequencerBakeControlInSpace(UControlRig* Co
 			ControlRig->Evaluate_AnyThread();
 
 			FMovieSceneSequenceTransform RootToLocalTransform = Sequencer->GetFocusedMovieSceneSequenceTransform();
+			FMovieSceneInverseSequenceTransform LocalToRootTransform = RootToLocalTransform.Inverse();
 
 			for (int32 Index = 0; Index < Frames.Num(); ++Index)
 			{
@@ -1102,7 +1106,7 @@ void FControlRigSpaceChannelHelpers::SequencerBakeControlInSpace(UControlRig* Co
 
 				//evaluate sequencer
 				FFrameTime GlobalTime(Frame);
-				GlobalTime = GlobalTime * RootToLocalTransform.InverseNoLooping();
+				GlobalTime = LocalToRootTransform.TryTransformTime(GlobalTime).Get(GlobalTime);
 
 				FMovieSceneContext SceneContext = FMovieSceneContext(FMovieSceneEvaluationRange(GlobalTime, TickResolution), Sequencer->GetPlaybackStatus()).SetHasJumped(true);
 				Sequencer->GetEvaluationTemplate().EvaluateSynchronousBlocking(SceneContext);
@@ -1142,7 +1146,7 @@ void FControlRigSpaceChannelHelpers::SequencerBakeControlInSpace(UControlRig* Co
 
 				//evaluate sequencer
 				FFrameTime GlobalTime(EndFrame);
-				GlobalTime = GlobalTime * RootToLocalTransform.InverseNoLooping();
+				GlobalTime = LocalToRootTransform.TryTransformTime(GlobalTime).Get(GlobalTime);
 
 				FMovieSceneContext SceneContext = FMovieSceneContext(FMovieSceneEvaluationRange(GlobalTime, TickResolution), Sequencer->GetPlaybackStatus()).SetHasJumped(true);
 				Sequencer->GetEvaluationTemplate().EvaluateSynchronousBlocking(SceneContext);
@@ -1532,7 +1536,7 @@ FReply FControlRigSpaceChannelHelpers::OpenBakeDialog(ISequencer* Sequencer, FMo
 	return FReply::Unhandled();
 }
 
-TUniquePtr<FCurveModel> CreateCurveEditorModel(const TMovieSceneChannelHandle<FMovieSceneControlRigSpaceChannel>& ChannelHandle, UMovieSceneSection* OwningSection, TSharedRef<ISequencer> InSequencer)
+TUniquePtr<FCurveModel> CreateCurveEditorModel(const TMovieSceneChannelHandle<FMovieSceneControlRigSpaceChannel>& ChannelHandle, const UE::Sequencer::FCreateCurveEditorModelParams& Params)
 {
 	if (FMovieSceneControlRigSpaceChannel* Channel = ChannelHandle.Get())
 	{
@@ -1541,7 +1545,7 @@ TUniquePtr<FCurveModel> CreateCurveEditorModel(const TMovieSceneChannelHandle<FM
 			const UCurveEditorSettings* Settings = GetDefault<UCurveEditorSettings>();
 			if (Settings == nullptr || Settings->GetShowBars())
 			{
-				return MakeUnique<FControlRigSpaceChannelCurveModel>(ChannelHandle, OwningSection, InSequencer);
+				return MakeUnique<FControlRigSpaceChannelCurveModel>(ChannelHandle, Params.OwningSection, Params.Sequencer);
 			}
 		}
 	}
diff --git a/Engine/Plugins/Animation/ControlRig/Source/ControlRigEditor/Private/Tools/ControlRigSnapper.cpp b/Engine/Plugins/Animation/ControlRig/Source/ControlRigEditor/Private/Tools/ControlRigSnapper.cpp
--- a/Engine/Plugins/Animation/ControlRig/Source/ControlRigEditor/Private/Tools/ControlRigSnapper.cpp
+++ b/Engine/Plugins/Animation/ControlRig/Source/ControlRigEditor/Private/Tools/ControlRigSnapper.cpp
@@ -132,7 +132,9 @@ static bool LocalGetControlRigControlTransforms(IMovieScenePlayer* Player, const
 		const FConstraintsManagerController& Controller = FConstraintsManagerController::Get(World);
 		FFrameRate TickResolution = MovieScene->GetTickResolution();
 		FFrameRate DisplayRate = MovieScene->GetDisplayRate();
-		
+
+		FMovieSceneInverseSequenceTransform LocalToRootTransform = RootToLocalTransform.Inverse();
+
 		OutTransforms.SetNum(Frames.Num());
 		for (int32 Index = 0; Index < Frames.Num(); ++Index)
 		{
@@ -140,8 +142,7 @@ static bool LocalGetControlRigControlTransforms(IMovieScenePlayer* Player, const
 			double DeltaTime = 0.0;
 			if (CurrentFrame.IsSet() == false || CurrentFrame.GetValue() != FrameNumber)
 			{
-				FFrameTime GlobalTime(FrameNumber);
-				GlobalTime = GlobalTime * RootToLocalTransform.InverseNoLooping(); //player evals in root time so need to go back to it.
+				FFrameTime GlobalTime = LocalToRootTransform.TryTransformTime(FrameNumber).Get(FrameNumber); //player evals in root time so need to go back to it.
 
 				FMovieSceneContext Context = FMovieSceneContext(FMovieSceneEvaluationRange(GlobalTime, TickResolution), Player->GetPlaybackStatus()).SetHasJumped(true);
 
@@ -267,7 +268,7 @@ struct FGuidAndActor
 			UMovieScene* MovieScene = Sequencer->GetFocusedMovieSceneSequence()->GetMovieScene();
 			FFrameRate TickResolution = MovieScene->GetTickResolution();
 			FFrameRate DisplayRate = MovieScene->GetDisplayRate();
-			FMovieSceneSequenceTransform RootToLocalTransform = Sequencer->GetFocusedMovieSceneSequenceTransform();
+			FMovieSceneInverseSequenceTransform LocalToRootTransform = Sequencer->GetFocusedMovieSceneSequenceTransform().Inverse();
 
 			//adjust keys for constraint
 			const FConstraintsManagerController& Controller = FConstraintsManagerController::Get(Actor->GetWorld());
@@ -277,8 +278,8 @@ struct FGuidAndActor
 				const FFrameNumber Frame = Frames[Index];
 				FTransform& CurrentTransform = LocalTransforms[Index];
 
-				FFrameTime GlobalTime(Frame);
-				GlobalTime = GlobalTime * RootToLocalTransform.InverseNoLooping();
+				FFrameTime GlobalTime = LocalToRootTransform.TryTransformTime(Frame).Get(Frame);
+
 				FMovieSceneContext Context = FMovieSceneContext(FMovieSceneEvaluationRange(GlobalTime, TickResolution), Sequencer->GetPlaybackStatus()).SetHasJumped(true);
 				if (Index == 0) // similar with baking first time in we need to evaluate twice (think due to double buffering that happens with skel mesh components).
 				{
diff --git a/Engine/Plugins/Animation/ControlRig/Source/ControlRigEditor/Public/ControlRigSpaceChannelEditors.h b/Engine/Plugins/Animation/ControlRig/Source/ControlRigEditor/Public/ControlRigSpaceChannelEditors.h
--- a/Engine/Plugins/Animation/ControlRig/Source/ControlRigEditor/Public/ControlRigSpaceChannelEditors.h
+++ b/Engine/Plugins/Animation/ControlRig/Source/ControlRigEditor/Public/ControlRigSpaceChannelEditors.h
@@ -70,8 +70,8 @@ struct FControlRigSpaceChannelHelpers
 FKeyHandle AddOrUpdateKey(FMovieSceneControlRigSpaceChannel* Channel, UMovieSceneSection* SectionToKey, FFrameNumber Time, ISequencer& Sequencer, const FGuid& ObjectBindingID, FTrackInstancePropertyBindings* PropertyBindings);
 
 /** Key editor overrides */
-bool CanCreateKeyEditor(const FMovieSceneControlRigSpaceChannel*       Channel);
-TSharedRef<SWidget> CreateKeyEditor(const TMovieSceneChannelHandle<FMovieSceneControlRigSpaceChannel>&        Channel, UMovieSceneSection* Section, const FGuid& InObjectBindingID, TWeakPtr<FTrackInstancePropertyBindings> PropertyBindings, TWeakPtr<ISequencer> Sequencer);
+bool CanCreateKeyEditor(const FMovieSceneControlRigSpaceChannel* Channel);
+TSharedRef<SWidget> CreateKeyEditor(const TMovieSceneChannelHandle<FMovieSceneControlRigSpaceChannel>& Channel, const UE::Sequencer::FCreateKeyEditorParams& Params);
 
 /** Key drawing overrides */
 void DrawKeys(FMovieSceneControlRigSpaceChannel* Channel, TArrayView<const FKeyHandle> InKeyHandles, const UMovieSceneSection* InOwner, TArrayView<FKeyDrawParams> OutKeyDrawParams);
@@ -87,6 +87,6 @@ int32 DrawExtra(FMovieSceneControlRigSpaceChannel* Channel, const UMovieSceneSec
 
 /** Curve editor models */
 inline bool SupportsCurveEditorModels(const TMovieSceneChannelHandle<FMovieSceneControlRigSpaceChannel>& Channel) { return true; }
-TUniquePtr<FCurveModel> CreateCurveEditorModel(const TMovieSceneChannelHandle<FMovieSceneControlRigSpaceChannel>& Channel, UMovieSceneSection* OwningSection, TSharedRef<ISequencer> InSequencer);
+TUniquePtr<FCurveModel> CreateCurveEditorModel(const TMovieSceneChannelHandle<FMovieSceneControlRigSpaceChannel>& Channel, const UE::Sequencer::FCreateCurveEditorModelParams& Params);
 
 
diff --git a/Engine/Plugins/Editor/SequencerAnimTools/Source/SequencerAnimTools/Private/Sequencer/MovieSceneTransformTrail.cpp b/Engine/Plugins/Editor/SequencerAnimTools/Source/SequencerAnimTools/Private/Sequencer/MovieSceneTransformTrail.cpp
--- a/Engine/Plugins/Editor/SequencerAnimTools/Source/SequencerAnimTools/Private/Sequencer/MovieSceneTransformTrail.cpp
+++ b/Engine/Plugins/Editor/SequencerAnimTools/Source/SequencerAnimTools/Private/Sequencer/MovieSceneTransformTrail.cpp
@@ -739,15 +739,16 @@ bool FMovieSceneControlRigTransformTrail::ApplyDelta(const FVector& Pos, const F
 		IMovieScenePlayer* Player = GetSequencer().Get();
 
 		FMovieSceneSequenceTransform RootToLocalTransform = GetSequencer()->GetFocusedMovieSceneSequenceTransform();
+		FMovieSceneInverseSequenceTransform LocalToRootTransform = RootToLocalTransform.Inverse();
 
-		auto EvalControlRig = [&Context,&RootToLocalTransform, Pos, TickResolution, ControlRig,Player, this](FTrailKeyInfo* KeyInfo)
+		auto EvalControlRig = [&Context,&LocalToRootTransform, Pos, TickResolution, ControlRig,Player, this](FTrailKeyInfo* KeyInfo)
 		{
 			if (KeyInfo)
 			{
 				Context.LocalTime = TickResolution.AsSeconds(FFrameTime(KeyInfo->FrameNumber));
 				Context.KeyMask = (uint32)EControlRigContextChannelToKey::Translation;
 				FFrameTime GlobalTime(KeyInfo->FrameNumber);
-				GlobalTime = GlobalTime * RootToLocalTransform.InverseNoLooping(); //player evals in root time so need to go back to it.
+				GlobalTime = LocalToRootTransform.TryTransformTime(GlobalTime).Get(GlobalTime); //player evals in root time so need to go back to it.
 
 				FMovieSceneContext MovieSceneContext = FMovieSceneContext(FMovieSceneEvaluationRange(GlobalTime, TickResolution), Player->GetPlaybackStatus()).SetHasJumped(true);
 
@@ -805,8 +806,8 @@ bool FMovieSceneControlRigTransformTrail::EndTracking()
 		IMovieScenePlayer* Player = GetSequencer().Get();
 		FFrameTime StartTime = GetSequencer()->GetLocalTime().Time;
 		FFrameRate TickResolution = GetSequencer()->GetFocusedTickResolution();
-		FMovieSceneSequenceTransform RootToLocalTransform = GetSequencer()->GetFocusedMovieSceneSequenceTransform();
-		StartTime = StartTime * RootToLocalTransform.InverseNoLooping(); //player evals in root time so need to go back to it.
+		FMovieSceneInverseSequenceTransform LocalToRootTransform = GetSequencer()->GetFocusedMovieSceneSequenceTransform().Inverse();
+		StartTime = LocalToRootTransform.TryTransformTime(StartTime).Get(StartTime); //player evals in root time so need to go back to it.
 
 		FMovieSceneContext MovieSceneContext = FMovieSceneContext(FMovieSceneEvaluationRange(StartTime, TickResolution), Player->GetPlaybackStatus()).SetHasJumped(true);
 		
@@ -844,8 +845,9 @@ bool FMovieSceneControlRigTransformTrail::HandleAltClick(FEditorViewportClient*
 	FFrameTime DisplayTime = FFrameRate::TransformTime(GlobalTime, DisplayResolution, TickResolution);
 	DisplayTime = DisplayTime.RoundToFrame();
 	GlobalTime = FFrameRate::TransformTime(DisplayTime, TickResolution, DisplayResolution);
-	FMovieSceneSequenceTransform RootToLocalTransform = GetSequencer()->GetFocusedMovieSceneSequenceTransform();
-	GlobalTime = GlobalTime * RootToLocalTransform.InverseNoLooping(); //player evals in root time so need to go back to it.
+
+	FMovieSceneInverseSequenceTransform LocalToRootTransform = GetSequencer()->GetFocusedMovieSceneSequenceTransform().Inverse();
+	GlobalTime = LocalToRootTransform.TryTransformTime(GlobalTime).Get(GlobalTime); //player evals in root time so need to go back to it.
 
 	Context.LocalTime = TickResolution.AsSeconds(GlobalTime);
 	Context.KeyMask = (uint32)EControlRigContextChannelToKey::Translation;
@@ -856,8 +858,7 @@ bool FMovieSceneControlRigTransformTrail::HandleAltClick(FEditorViewportClient*
 	ControlRig->SetControlGlobalTransform(ControlName, NewTransform, true, Context, false /*undo*/, false /*bPrintPython*/, true/* bFixEulerFlips*/);
 
 	//eval back at current time
-	FFrameTime StartTime = GetSequencer()->GetLocalTime().Time;
-	StartTime = StartTime * RootToLocalTransform.InverseNoLooping(); //player evals in root time so need to go back to it.
+	FFrameTime StartTime = GetSequencer()->GetGlobalTime().Time;
 
 	MovieSceneContext = FMovieSceneContext(FMovieSceneEvaluationRange(StartTime, TickResolution), Player->GetPlaybackStatus()).SetHasJumped(true);
 	Player->GetEvaluationTemplate().EvaluateSynchronousBlocking(MovieSceneContext);
diff --git a/Engine/Plugins/Editor/SequencerAnimTools/Source/SequencerAnimTools/Private/Sequencer/SequencerTrailHierarchy.cpp b/Engine/Plugins/Editor/SequencerAnimTools/Source/SequencerAnimTools/Private/Sequencer/SequencerTrailHierarchy.cpp
--- a/Engine/Plugins/Editor/SequencerAnimTools/Source/SequencerAnimTools/Private/Sequencer/SequencerTrailHierarchy.cpp
+++ b/Engine/Plugins/Editor/SequencerAnimTools/Source/SequencerAnimTools/Private/Sequencer/SequencerTrailHierarchy.cpp
@@ -225,7 +225,7 @@ void FSequencerTrailHierarchy::UpdateControlRig(const TArray<FFrameNumber> &Fram
 		FMovieSceneSequencePlaybackSettings Settings;
 		FLevelSequenceCameraSettings CameraSettings;
 		FMovieSceneSequenceIDRef Template = MovieSceneSequenceID::Root;
-		FMovieSceneSequenceTransform RootToLocalTransform = Sequencer->GetFocusedMovieSceneSequenceTransform();
+		FMovieSceneInverseSequenceTransform LocalToRootTransform = Sequencer->GetFocusedMovieSceneSequenceTransform().Inverse();
 		IMovieScenePlayer* Player = Sequencer.Get();
 		ULevelSequence* LevelSequence = Cast<ULevelSequence>(Sequencer->GetFocusedMovieSceneSequence());
 
@@ -256,7 +256,7 @@ void FSequencerTrailHierarchy::UpdateControlRig(const TArray<FFrameNumber> &Fram
 				{
 					const FFrameNumber& FrameNumber = Frames[Index];
 					FFrameTime GlobalTime(FrameNumber);
-					GlobalTime = GlobalTime * RootToLocalTransform.InverseNoLooping(); //player evals in root time so need to go back to it.
+					GlobalTime = LocalToRootTransform.TryTransformTime(GlobalTime).Get(GlobalTime);
 
 					FMovieSceneContext Context = FMovieSceneContext(FMovieSceneEvaluationRange(GlobalTime, TickResolution), Player->GetPlaybackStatus()).SetHasJumped(true);
 
@@ -290,8 +290,7 @@ void FSequencerTrailHierarchy::UpdateControlRig(const TArray<FFrameNumber> &Fram
 		}
 		if (Player)
 		{
-			FFrameTime StartTime = Sequencer->GetLocalTime().Time;
-			StartTime = StartTime * RootToLocalTransform.InverseNoLooping(); //player evals in root time so need to go back to it.
+			FFrameTime StartTime = Sequencer->GetGlobalTime().Time;
 			FMovieSceneContext Context = FMovieSceneContext(FMovieSceneEvaluationRange(StartTime, TickResolution), Player->GetPlaybackStatus()).SetHasJumped(true);
 			Player->GetEvaluationTemplate().EvaluateSynchronousBlocking(Context);
 			ControlRig->Evaluate_AnyThread();
diff --git a/Engine/Plugins/Experimental/CinematicPrestreaming/Source/CinematicPrestreamingEditor/Private/CinePrestreamingRecorderSetting.cpp b/Engine/Plugins/Experimental/CinematicPrestreaming/Source/CinematicPrestreamingEditor/Private/CinePrestreamingRecorderSetting.cpp
--- a/Engine/Plugins/Experimental/CinematicPrestreaming/Source/CinematicPrestreamingEditor/Private/CinePrestreamingRecorderSetting.cpp
+++ b/Engine/Plugins/Experimental/CinematicPrestreaming/Source/CinematicPrestreamingEditor/Private/CinePrestreamingRecorderSetting.cpp
@@ -247,7 +247,7 @@ void UCinePrestreamingRecorderSetting::CreateAssetsFromData()
 
 			// We only use the leafmost moviescene as that's the one that actually gets modified.
 			OutputSegment->MovieScene = Node->MovieScene.Get();
-			OutputSegment->Range =  ShotInfo.OuterToInnerTransform.TransformRangeUnwarped(ShotInfo.TotalOutputRangeRoot);
+			OutputSegment->Range = UE::MovieScene::ConvertToDiscreteRange(ShotInfo.OuterToInnerTransform.ComputeTraversedHull(ShotInfo.TotalOutputRangeRoot));
 		}
 
 		// Convert our samples and append them to the existing array
diff --git a/Engine/Plugins/Experimental/PropertyAnimator/Source/PropertyAnimatorEditor/Private/Sequencer/PropertyAnimatorEditorCurveChannelInterface.h b/Engine/Plugins/Experimental/PropertyAnimator/Source/PropertyAnimatorEditor/Private/Sequencer/PropertyAnimatorEditorCurveChannelInterface.h
--- a/Engine/Plugins/Experimental/PropertyAnimator/Source/PropertyAnimatorEditor/Private/Sequencer/PropertyAnimatorEditorCurveChannelInterface.h
+++ b/Engine/Plugins/Experimental/PropertyAnimator/Source/PropertyAnimatorEditor/Private/Sequencer/PropertyAnimatorEditorCurveChannelInterface.h
@@ -31,12 +31,12 @@ struct TPropertyAnimatorEditorCurveChannelInterface : ISequencerChannelInterface
 	virtual void DrawKeys_Raw(FMovieSceneChannel*, TConstArrayView<FKeyHandle>, const UMovieSceneSection*, TArrayView<FKeyDrawParams>) const override {}
 	virtual bool ShouldShowCurve_Raw(const FMovieSceneChannel*, UMovieSceneSection*) const override { return true; }
 	virtual bool SupportsCurveEditorModels_Raw(const FMovieSceneChannelHandle&) const override { return false; }
-	virtual TUniquePtr<FCurveModel> CreateCurveEditorModel_Raw(const FMovieSceneChannelHandle&, UMovieSceneSection*, TSharedRef<ISequencer>) const override { return nullptr; }
+	virtual TUniquePtr<FCurveModel> CreateCurveEditorModel_Raw(const FMovieSceneChannelHandle&, const UE::Sequencer::FCreateCurveEditorModelParams& Params) const override { return nullptr; }
 	virtual TSharedPtr<UE::Sequencer::FChannelModel> CreateChannelModel_Raw(const FMovieSceneChannelHandle&, FName) const override { return nullptr; }
 	virtual TSharedPtr<UE::Sequencer::STrackAreaLaneView> CreateChannelView_Raw(const FMovieSceneChannelHandle&, TWeakPtr<UE::Sequencer::FChannelModel>, const UE::Sequencer::FCreateTrackLaneViewParams&) const override { return nullptr; }
 	virtual void ExtendSectionMenu_Raw(FMenuBuilder& InMenuBuilder, TSharedPtr<FExtender> InMenuExtender, TConstArrayView<FMovieSceneChannelHandle> InChannels, TConstArrayView<UMovieSceneSection*> InSections, TWeakPtr<ISequencer> InSequencer) const override;
 	virtual int32 DrawExtra_Raw(FMovieSceneChannel* InChannel, const UMovieSceneSection* InOwner, const FSequencerChannelPaintArgs& InPaintArgs, int32 InLayerId) const override;
-	virtual TSharedRef<SWidget> CreateKeyEditor_Raw(const FMovieSceneChannelHandle& InChannel, UMovieSceneSection* InSection, const FGuid& InObjectBindingID, TWeakPtr<FTrackInstancePropertyBindings> InPropertyBindings, TWeakPtr<ISequencer> InSequencer) const override;
+	virtual TSharedRef<SWidget> CreateKeyEditor_Raw(const FMovieSceneChannelHandle& InChannel, const UE::Sequencer::FCreateKeyEditorParams& Params) const override;
 	//~ End ISequencerChannelInterface
 };
 
@@ -100,10 +100,7 @@ int32 TPropertyAnimatorEditorCurveChannelInterface<InChannelType, InMenuExtensio
 
 template<typename InChannelType, typename InMenuExtensionType>
 TSharedRef<SWidget> TPropertyAnimatorEditorCurveChannelInterface<InChannelType, InMenuExtensionType>::CreateKeyEditor_Raw(const FMovieSceneChannelHandle& InChannel
-	, UMovieSceneSection* InSection
-	, const FGuid& InObjectBindingID
-	, TWeakPtr<FTrackInstancePropertyBindings> InPropertyBindings
-	, TWeakPtr<ISequencer> InSequencer) const
+	, const UE::Sequencer::FCreateKeyEditorParams& Params) const
 {
 	const TMovieSceneExternalValue<FCurveValueType>* ExternalValue = InChannel.Cast<FChannelType>().GetExtendedEditorData();
 	if (!ExternalValue)
@@ -111,10 +108,10 @@ TSharedRef<SWidget> TPropertyAnimatorEditorCurveChannelInterface<InChannelType,
 		return SNullWidget::NullWidget;
 	}
 
-	return SNew(SKeyEditorWidget, FKeyEditor(InObjectBindingID
+	return SNew(SKeyEditorWidget, FKeyEditor(Params.ObjectBindingID
 		, InChannel.Cast<FChannelType>()
-		, InSection
-		, InSequencer
-		, InPropertyBindings
+		, Params.OwningSection
+		, Params.Sequencer
+		, Params.PropertyBindings
 		, ExternalValue->OnGetExternalValue));
 }
diff --git a/Engine/Plugins/Importers/USDImporter/Source/USDStage/Private/USDLevelSequenceHelper.cpp b/Engine/Plugins/Importers/USDImporter/Source/USDStage/Private/USDLevelSequenceHelper.cpp
--- a/Engine/Plugins/Importers/USDImporter/Source/USDStage/Private/USDLevelSequenceHelper.cpp
+++ b/Engine/Plugins/Importers/USDImporter/Source/USDStage/Private/USDLevelSequenceHelper.cpp
@@ -3519,8 +3519,12 @@ void FUsdLevelSequenceHelperImpl::UpdateUsdLayerOffsetFromSection(const UMovieSc
 	// This will obviously be quantized to frame intervals for now
 	double SubSectionStartTimeCode = TickResolution.AsSeconds(ModifiedStartFrame) * TimeCodesPerSecond;
 
+	const double FixedPlayRate = Section->Parameters.TimeScale.GetType() == EMovieSceneTimeWarpType::FixedPlayRate\
+		? Section->Parameters.TimeScale.AsFixedPlayRate()
+		: 1.f;
+
 	UE::FSdfLayerOffset NewLayerOffset;
-	NewLayerOffset.Scale = FMath::IsNearlyZero(Section->Parameters.TimeScale) ? 0.f : 1.f / Section->Parameters.TimeScale;
+	NewLayerOffset.Scale = FMath::IsNearlyZero(FixedPlayRate) ? 0.f : 1.f / FixedPlayRate;
 	NewLayerOffset.Offset = SubSectionStartTimeCode - SubStartTimeCode * NewLayerOffset.Scale;
 
 	if (FMath::IsNearlyZero(NewLayerOffset.Offset))
@@ -4267,7 +4271,7 @@ void FUsdLevelSequenceHelperImpl::HandleControlRigSectionChange(UMovieSceneContr
 	const UsdUtils::FBlendShapeMap& BlendShapeMap = StageActorValue->GetBlendShapeMap();
 	bool bBaked = UnrealToUsd::ConvertControlRigSection(
 		&Section,
-		SequenceTransform.InverseNoLooping(),
+		SequenceTransform.Inverse(),
 		MovieScene,
 		Player,
 		Skeleton->GetReferenceSkeleton(),
diff --git a/Engine/Plugins/Interchange/Runtime/Source/Import/Private/Animation/InterchangeLevelSequenceFactory.cpp b/Engine/Plugins/Interchange/Runtime/Source/Import/Private/Animation/InterchangeLevelSequenceFactory.cpp
--- a/Engine/Plugins/Interchange/Runtime/Source/Import/Private/Animation/InterchangeLevelSequenceFactory.cpp
+++ b/Engine/Plugins/Interchange/Runtime/Source/Import/Private/Animation/InterchangeLevelSequenceFactory.cpp
@@ -419,9 +419,10 @@ namespace UE::Interchange::Private
 
 		// Internally AddSequenceOnRow will automatically bump overlapping subsequences, so we can just add where it's ideal for us
 		UMovieSceneSubSection* NewSection = SubTrack->AddSequenceOnRow(TargetMovieSceneSequence, DstLowerBound, DstUpperBound.Value - DstLowerBound.Value, INDEX_NONE);
-		
-		NewSection->Parameters.TimeScale = 1.f;
-		InstanceNode.GetCustomTimeScale(NewSection->Parameters.TimeScale);
+
+		float TimeScale = 1.f;
+		InstanceNode.GetCustomTimeScale(TimeScale);
+		NewSection->Parameters.TimeScale.Set(TimeScale);
 		
 		int32 CompletionMode;
 		if (InstanceNode.GetCustomCompletionMode(CompletionMode))
diff --git a/Engine/Plugins/MovieScene/LevelSequenceEditor/Source/LevelSequenceEditor/Private/LevelSequenceEditorSubsystem.cpp b/Engine/Plugins/MovieScene/LevelSequenceEditor/Source/LevelSequenceEditor/Private/LevelSequenceEditorSubsystem.cpp
--- a/Engine/Plugins/MovieScene/LevelSequenceEditor/Source/LevelSequenceEditor/Private/LevelSequenceEditorSubsystem.cpp
+++ b/Engine/Plugins/MovieScene/LevelSequenceEditor/Source/LevelSequenceEditor/Private/LevelSequenceEditorSubsystem.cpp
@@ -1639,16 +1639,20 @@ bool ULevelSequenceEditorSubsystem::BakeTransformWithSettings(const TArray<FMovi
 	}
 	CalculateFramesPerGuid(Sequencer, SettingsInTick, BakeDataMap, TotalFrameMap);
 
-	FMovieSceneSequenceTransform RootToLocalTransform = Sequencer->GetFocusedMovieSceneSequenceTransform();
+	FMovieSceneInverseSequenceTransform LocalToRootTransform = Sequencer->GetFocusedMovieSceneSequenceTransform().Inverse();
 	
 	TArray<FFrameNumber> AllFrames;
 	TotalFrameMap.GenerateKeyArray(AllFrames);
 
 	for (FFrameNumber KeyTime: AllFrames)
 	{
-		FMovieSceneEvaluationRange Range(KeyTime * RootToLocalTransform.InverseNoLooping(), TickResolution);
+		TOptional<FFrameTime> NewGlobalTime = LocalToRootTransform.TryTransformTime(KeyTime);
+		if (!NewGlobalTime)
+		{
+			continue;
+		}
 
-		Sequencer->SetGlobalTime(Range.GetTime());
+		Sequencer->SetGlobalTime(NewGlobalTime.GetValue());
 
 		for (const FMovieSceneBindingProxy& ObjectBinding : ObjectBindings)
 		{
diff --git a/Engine/Plugins/MovieScene/MovieRenderPipeline/Source/MovieRenderPipelineCore/Private/MoviePipelineBlueprintLibrary.cpp b/Engine/Plugins/MovieScene/MovieRenderPipeline/Source/MovieRenderPipelineCore/Private/MoviePipelineBlueprintLibrary.cpp
--- a/Engine/Plugins/MovieScene/MovieRenderPipeline/Source/MovieRenderPipelineCore/Private/MoviePipelineBlueprintLibrary.cpp
+++ b/Engine/Plugins/MovieScene/MovieRenderPipeline/Source/MovieRenderPipelineCore/Private/MoviePipelineBlueprintLibrary.cpp
@@ -31,6 +31,7 @@
 #include "Misc/StringBuilder.h"
 #include "MovieSceneSequenceVisitor.h"
 #include "MovieSceneSequenceID.h"
+#include "MovieSceneTransformTypes.h"
 #include "MoviePipelineUtils.h"
 #include "UObject/UObjectHash.h"
 #include "MovieRenderPipelineCoreModule.h"
@@ -347,6 +348,8 @@ UMoviePipelineQueue* UMoviePipelineBlueprintLibrary::LoadManifestFileFromString(
 
 void UMoviePipelineBlueprintLibrary::UpdateJobShotListFromSequence(ULevelSequence* InSequence, UMoviePipelineExecutorJob* InJob, bool& bShotsChanged)
 {
+	using namespace UE::MovieScene;
+
 	if (!ensureMsgf(InSequence && InJob, TEXT("Cannot generate shot list for null sequence/job.")))
 	{
 		return;
@@ -359,6 +362,8 @@ void UMoviePipelineBlueprintLibrary::UpdateJobShotListFromSequence(ULevelSequenc
 		{
 			TArray<FSequenceRangeInfo> OutRanges;
 
+			FMovieSceneInverseSequenceTransform LocalToRoot = LocalSpace.RootToSequenceTransform.Inverse();
+
 			for (UMovieSceneSection* Section : InTrack->GetAllSections())
 			{
 				if(!Section)
@@ -378,20 +383,26 @@ void UMoviePipelineBlueprintLibrary::UpdateJobShotListFromSequence(ULevelSequenc
 				UMovieScene* OwningScene = Section->GetTypedOuter<UMovieScene>();
 				TRange<FFrameNumber> LocalCameraRange = Section->GetRange();
 
-				// Intersect it with the root range so that if the parent has trimmed down the sub-section we don't render outside that.
-				TRange<FFrameNumber> RootCameraRange = TRange<FFrameNumber>::Intersection(LocalSpace.RootClampRange, LocalSpace.RootToSequenceTransform.InverseNoLooping().TransformRangeUnwarped(LocalCameraRange));
-				if (!RootCameraRange.IsEmpty())
+				auto Visit = [&LocalSpace, Section, InTree](TRange<FFrameTime> InRange)
 				{
-					if (UMovieSceneCinematicShotSection* ShotSection = Cast<UMovieSceneCinematicShotSection>(Section))
-					{
-						// We add the child sequence here so that the hierarchy tree gets built correctly.
-						InTree->Add(RootCameraRange, MakeTuple(Section, ShotSection->GetSequenceID(), LocalSpace.HierarchicalBias));
-					}
-					else
+					// Intersect it with the root range so that if the parent has trimmed down the sub-section we don't render outside that.
+					TRange<FFrameNumber> RootCameraRange = TRange<FFrameNumber>::Intersection(LocalSpace.RootClampRange, ConvertToDiscreteRange(InRange));
+					if (!RootCameraRange.IsEmpty())
 					{
-						InTree->Add(RootCameraRange, MakeTuple(Section, LocalSpace.SequenceID, LocalSpace.HierarchicalBias));
+						if (UMovieSceneCinematicShotSection* ShotSection = Cast<UMovieSceneCinematicShotSection>(Section))
+						{
+							// We add the child sequence here so that the hierarchy tree gets built correctly.
+							InTree->Add(RootCameraRange, MakeTuple(Section, ShotSection->GetSequenceID(), LocalSpace.HierarchicalBias));
+						}
+						else
+						{
+							InTree->Add(RootCameraRange, MakeTuple(Section, LocalSpace.SequenceID, LocalSpace.HierarchicalBias));
+						}
 					}
-				}
+					return true;
+				};
+
+				LocalToRoot.TransformFiniteRangeWithinRange(ConvertToFrameTimeRange(LocalCameraRange), Visit, LocalSpace.StartBreadcrumbs, LocalSpace.EndBreadcrumbs);
 			}
 		}
 
@@ -431,7 +442,7 @@ void UMoviePipelineBlueprintLibrary::UpdateJobShotListFromSequence(ULevelSequenc
 			TTuple<FString, FString> Name;
 			TSharedPtr<MoviePipeline::FCameraCutSubSectionHierarchyNode> LeafNode;
 			TRange<FFrameNumber> CameraCutWarmUpRange;
-			FMovieSceneSequenceTransform InnerToOuterTransform;
+			FMovieSceneSequenceTransform RootToInnerTransform;
 		};
 
 		TArray<FLinearizedEntity> Entities;
@@ -628,20 +639,34 @@ void UMoviePipelineBlueprintLibrary::UpdateJobShotListFromSequence(ULevelSequenc
 				}
 			}
 
-			FMovieSceneSequenceTransform InnerToOuterTransform;
+			// To make the camera cut range detection more consistent, we'll convert the start of this Entity into root sequence space, and then we'll convert the start of the
+			// camera cut into root sequence space, and the difference between them is the range to use. This is a more reliable way than looking at the actual Playback Range.
+			TRange<FFrameNumber> EntityRangeInRoot = Entity.Range;
+			TRange<FFrameNumber> CameraCutRangeInRoot = EntityRangeInRoot;
+
 			FMovieSceneSubSequenceData* SubSequenceData = SequenceHierarchyCache.FindSubData(Entity.SequenceID);
 			if (SubSequenceData)
 			{
-				InnerToOuterTransform = SubSequenceData->RootToSequenceTransform.InverseFromAllFirstLoops();
+				Entity.RootToInnerTransform = SubSequenceData->RootToSequenceTransform;
 			}
 
-			// To make the camera cut range detection more consistent, we'll convert the start of this Entity into root sequence space, and then we'll convert the start of the
-			// camera cut into root sequence space, and the difference between them is the range to use. This is a more reliable way than looking at the actual Playback Range.
-			TRange<FFrameNumber> EntityRangeInRoot = Entity.Range;
-			TRange<FFrameNumber> CameraCutRangeInRoot = EntityRangeInRoot;
 			if (LeafNode->CameraCutSection.IsValid())
 			{
-				CameraCutRangeInRoot = InnerToOuterTransform.TransformRangeConstrained(LeafNode->CameraCutSection->GetRange());
+				CameraCutRangeInRoot = LeafNode->CameraCutSection->GetRange();
+
+				// Put the camera cut range in root space
+				FMovieSceneInverseSequenceTransform InnerToRoot = Entity.RootToInnerTransform.Inverse();
+
+				FMovieSceneTransformBreadcrumbs RangeStartBreadcrumbs, RangeEndBreadcrumbs;
+				Entity.RootToInnerTransform.TransformTime(UE::MovieScene::DiscreteInclusiveLower(EntityRangeInRoot), UE::MovieScene::FTransformTimeParams().HarvestBreadcrumbs(RangeStartBreadcrumbs));
+				Entity.RootToInnerTransform.TransformTime(UE::MovieScene::DiscreteExclusiveUpper(EntityRangeInRoot), UE::MovieScene::FTransformTimeParams().HarvestBreadcrumbs(RangeEndBreadcrumbs));
+
+				auto Visit = [&CameraCutRangeInRoot](TRange<FFrameTime> InRange)
+				{
+					CameraCutRangeInRoot = ConvertToDiscreteRange(InRange);
+					return false; // only call this once
+				};
+				InnerToRoot.TransformFiniteRangeWithinRange(ConvertToFrameTimeRange(CameraCutRangeInRoot), Visit, RangeStartBreadcrumbs, RangeEndBreadcrumbs);
 			}
 
 			TRange<FFrameNumber> CameraCutWarmUpRange = TRange<FFrameNumber>::Empty();
@@ -657,7 +682,6 @@ void UMoviePipelineBlueprintLibrary::UpdateJobShotListFromSequence(ULevelSequenc
 			Entity.CameraCutWarmUpRange = CameraCutWarmUpRange;
 			Entity.LeafNode = LeafNode;
 			Entity.Name = MoviePipeline::GetNameForShot(SequenceHierarchyCache, InSequence, LeafNode);
-			Entity.InnerToOuterTransform = InnerToOuterTransform;
 		}
 
 		// Fallback case for no shot sections detected and no camera cut sections detected.
@@ -676,7 +700,7 @@ void UMoviePipelineBlueprintLibrary::UpdateJobShotListFromSequence(ULevelSequenc
 			Entity.CameraCutWarmUpRange = TRange<FFrameNumber>::Empty();
 			Entity.LeafNode = LeafNode;
 			Entity.Name = MoviePipeline::GetNameForShot(SequenceHierarchyCache, InSequence, LeafNode);
-			Entity.InnerToOuterTransform = FMovieSceneSequenceTransform();
+			Entity.RootToInnerTransform = FMovieSceneSequenceTransform();
 		}
 
 		// We need to generate all of the linearized segments first so that we have all of the names available.
@@ -766,7 +790,7 @@ void UMoviePipelineBlueprintLibrary::UpdateJobShotListFromSequence(ULevelSequenc
 			NewShot->ShotInfo.SubSectionHierarchy = Entity.LeafNode;
 			NewShot->ShotInfo.TotalOutputRangeRoot = Entity.Range;
 			NewShot->ShotInfo.WarmupRangeRoot = Entity.CameraCutWarmUpRange;
-			NewShot->ShotInfo.OuterToInnerTransform = Entity.InnerToOuterTransform.InverseNoLooping();
+			NewShot->ShotInfo.OuterToInnerTransform = Entity.RootToInnerTransform;
 			NewShot->SidecarCameras = Entity.SidecarCameras;
 			UE_LOG(LogMovieRenderPipeline, Log, TEXT("Registering range: %s (InnerName: %s OuterName: %s)"), *LexToString(NewShot->ShotInfo.TotalOutputRangeRoot), *NewShot->InnerName, *NewShot->OuterName);
 		}
diff --git a/Engine/Plugins/MovieScene/MovieSceneTextTrack/Source/MovieSceneTextTrackEditor/Private/TextPropertyTrackEditor.cpp b/Engine/Plugins/MovieScene/MovieSceneTextTrack/Source/MovieSceneTextTrackEditor/Private/TextPropertyTrackEditor.cpp
--- a/Engine/Plugins/MovieScene/MovieSceneTextTrack/Source/MovieSceneTextTrackEditor/Private/TextPropertyTrackEditor.cpp
+++ b/Engine/Plugins/MovieScene/MovieSceneTextTrack/Source/MovieSceneTextTrackEditor/Private/TextPropertyTrackEditor.cpp
@@ -6,6 +6,7 @@
 #include "SequencerKeyStructGenerator.h"
 #include "UObject/TextProperty.h"
 #include "Widgets/STextKeyEditor.h"
+#include "ISequencerChannelInterface.h"
 
 TSharedRef<ISequencerTrackEditor> FTextPropertyTrackEditor::CreateTrackEditor(TSharedRef<ISequencer> OwningSequencer)
 {
@@ -31,11 +32,7 @@ bool CanCreateKeyEditor(const FMovieSceneTextChannel* Channel)
 	return true;
 }
 
-TSharedRef<SWidget> CreateKeyEditor(const TMovieSceneChannelHandle<FMovieSceneTextChannel>& Channel
-	, UMovieSceneSection* Section
-	, const FGuid& InObjectBindingID
-	, TWeakPtr<FTrackInstancePropertyBindings> PropertyBindings
-	, TWeakPtr<ISequencer> InSequencer)
+TSharedRef<SWidget> CreateKeyEditor(const TMovieSceneChannelHandle<FMovieSceneTextChannel>& Channel, const UE::Sequencer::FCreateKeyEditorParams& InParams)
 {
 	using namespace UE::MovieScene;
 
@@ -46,11 +43,11 @@ TSharedRef<SWidget> CreateKeyEditor(const TMovieSceneChannelHandle<FMovieSceneTe
 	}
 
 	FTextKeyEditorParams Params;
-	Params.ObjectBindingID = InObjectBindingID;
+	Params.ObjectBindingID = InParams.ObjectBindingID;
 	Params.ChannelHandle = Channel;
-	Params.WeakSection = Section;
-	Params.WeakSequencer = InSequencer;
-	Params.WeakPropertyBindings = PropertyBindings;
+	Params.WeakSection = InParams.OwningSection;
+	Params.WeakSequencer = InParams.Sequencer;
+	Params.WeakPropertyBindings = InParams.PropertyBindings;
 	Params.OnGetExternalValue = ExternalValue->OnGetExternalValue;
 
 	return SNew(STextKeyEditor, MoveTemp(Params));
diff --git a/Engine/Plugins/MovieScene/MovieSceneTextTrack/Source/MovieSceneTextTrackEditor/Private/TextPropertyTrackEditor.h b/Engine/Plugins/MovieScene/MovieSceneTextTrack/Source/MovieSceneTextTrackEditor/Private/TextPropertyTrackEditor.h
--- a/Engine/Plugins/MovieScene/MovieSceneTextTrack/Source/MovieSceneTextTrackEditor/Private/TextPropertyTrackEditor.h
+++ b/Engine/Plugins/MovieScene/MovieSceneTextTrack/Source/MovieSceneTextTrackEditor/Private/TextPropertyTrackEditor.h
@@ -8,6 +8,10 @@
 class UMovieSceneKeyStructType;
 class FSequencerKeyStructGenerator;
 struct FMovieSceneTextChannel;
+namespace UE::Sequencer
+{
+	struct FCreateKeyEditorParams;
+}
 
 /** Track Editor for Text Property */
 class FTextPropertyTrackEditor : public FPropertyTrackEditor<UMovieSceneTextTrack>
@@ -39,11 +43,7 @@ protected:
 
 bool CanCreateKeyEditor(const FMovieSceneTextChannel* Channel);
 
-TSharedRef<SWidget> CreateKeyEditor(const TMovieSceneChannelHandle<FMovieSceneTextChannel>& Channel
-	, UMovieSceneSection* Section
-	, const FGuid& InObjectBindingID
-	, TWeakPtr<FTrackInstancePropertyBindings> PropertyBindings
-	, TWeakPtr<ISequencer> InSequencer);
+TSharedRef<SWidget> CreateKeyEditor(const TMovieSceneChannelHandle<FMovieSceneTextChannel>& Channel, const UE::Sequencer::FCreateKeyEditorParams& Params);
 
 UMovieSceneKeyStructType* InstanceGeneratedStruct(FMovieSceneTextChannel* Channel, FSequencerKeyStructGenerator* Generator);
 
diff --git a/Engine/Plugins/MovieScene/SequencerScripting/Source/SequencerScripting/Private/ExtensionLibraries/MovieSceneSectionExtensions.cpp b/Engine/Plugins/MovieScene/SequencerScripting/Source/SequencerScripting/Private/ExtensionLibraries/MovieSceneSectionExtensions.cpp
--- a/Engine/Plugins/MovieScene/SequencerScripting/Source/SequencerScripting/Private/ExtensionLibraries/MovieSceneSectionExtensions.cpp
+++ b/Engine/Plugins/MovieScene/SequencerScripting/Source/SequencerScripting/Private/ExtensionLibraries/MovieSceneSectionExtensions.cpp
@@ -655,14 +655,18 @@ int32 UMovieSceneSectionExtensions::GetParentSequenceFrame(UMovieSceneSubSection
 
 	TArray<UMovieSceneSubSection*> SubSectionChain;
 	GetSubSectionChain(InSubSection, ParentSequence, SubSectionChain);
-		
+
 	FFrameRate LocalDisplayRate = InSubSection->GetSequence()->GetMovieScene()->GetDisplayRate();
 	FFrameRate LocalTickResolution = InSubSection->GetSequence()->GetMovieScene()->GetTickResolution();
 	FFrameTime LocalFrameTime = ConvertFrameTime(InFrame, LocalDisplayRate, LocalTickResolution);
-		
+
 	for (int32 SectionIndex = 0; SectionIndex < SubSectionChain.Num(); ++SectionIndex)
 	{
-		LocalFrameTime = LocalFrameTime * SubSectionChain[SectionIndex]->OuterToInnerTransform().InverseNoLooping();
+		TOptional<FFrameTime> NewLocalTime = SubSectionChain[SectionIndex]->OuterToInnerTransform().Inverse().TryTransformTime(LocalFrameTime);
+		if (NewLocalTime)
+		{
+			LocalFrameTime = NewLocalTime.GetValue();
+		}
 	}
 
 	FFrameRate ParentDisplayRate = ParentSequence->GetMovieScene()->GetDisplayRate();
diff --git a/Engine/Plugins/MovieScene/SequencerScripting/Source/SequencerScriptingEditor/Private/SequencerCurveEditorObject.cpp b/Engine/Plugins/MovieScene/SequencerScripting/Source/SequencerScriptingEditor/Private/SequencerCurveEditorObject.cpp
--- a/Engine/Plugins/MovieScene/SequencerScripting/Source/SequencerScriptingEditor/Private/SequencerCurveEditorObject.cpp
+++ b/Engine/Plugins/MovieScene/SequencerScripting/Source/SequencerScriptingEditor/Private/SequencerCurveEditorObject.cpp
@@ -87,7 +87,7 @@ TArray<FSequencerChannelProxy> USequencerCurveEditorObject::GetChannelsWithSelec
 		{
 			if (const FCurveModel* Curve = CurveEditor->FindCurve(Pair.Key))
 			{
-				if (UMovieSceneSection* Section = Cast<UMovieSceneSection>(Curve->GetOwningObject()))
+				if (UMovieSceneSection* Section = Curve->GetOwningObjectOrOuter<UMovieSceneSection>())
 				{
 					FName ChannelName = Curve->GetChannelName();
 					FSequencerChannelProxy ChannelProxy(ChannelName,Section);
@@ -112,7 +112,7 @@ TArray<int32> USequencerCurveEditorObject::GetSelectedKeys(const FSequencerChann
 		{
 			if (const FCurveModel* Curve = CurveEditor->FindCurve(Pair.Key))
 			{
-				if (UMovieSceneSection* Section = Cast<UMovieSceneSection>(Curve->GetOwningObject()))
+				if (UMovieSceneSection* Section = Curve->GetOwningObjectOrOuter<UMovieSceneSection>())
 				{
 					if (Section == ChannelProxy.Section)
 					{
diff --git a/Engine/Plugins/MovieScene/SequencerScripting/Source/SequencerScriptingEditor/Private/SequencerTools.cpp b/Engine/Plugins/MovieScene/SequencerScripting/Source/SequencerScriptingEditor/Private/SequencerTools.cpp
--- a/Engine/Plugins/MovieScene/SequencerScripting/Source/SequencerScriptingEditor/Private/SequencerTools.cpp
+++ b/Engine/Plugins/MovieScene/SequencerScripting/Source/SequencerScriptingEditor/Private/SequencerTools.cpp
@@ -255,7 +255,7 @@ bool ExportFBXInternal(const FSequencerExportFBXParams& InParams, UMovieSceneSeq
 	if (SubSequenceData)
 	{
 		RootToLocalTransform = SubSequenceData->RootToSequenceTransform;
-		StartTime = UE::MovieScene::DiscreteInclusiveLower(SubSequenceData->PlayRange.Value) * RootToLocalTransform.InverseNoLooping();
+		StartTime = RootToLocalTransform.Inverse().TryTransformTime(UE::MovieScene::DiscreteInclusiveLower(SubSequenceData->PlayRange.Value)).Get(StartTime);
 	}
 
 	bool bDidExport = false;
diff --git a/Engine/Plugins/MovieScene/TemplateSequence/Source/TemplateSequence/Private/Section/TemplateSequenceSection.cpp b/Engine/Plugins/MovieScene/TemplateSequence/Source/TemplateSequence/Private/Section/TemplateSequenceSection.cpp
--- a/Engine/Plugins/MovieScene/TemplateSequence/Source/TemplateSequence/Private/Section/TemplateSequenceSection.cpp
+++ b/Engine/Plugins/MovieScene/TemplateSequence/Source/TemplateSequence/Private/Section/TemplateSequenceSection.cpp
@@ -48,7 +48,7 @@ bool UTemplateSequenceSection::ShowCurveForChannel(const void *ChannelPtr) const
 void UTemplateSequenceSection::OnDilated(float DilationFactor, FFrameNumber Origin)
 {
 	// TODO-lchabant: shouldn't this be in the base class?
-	Parameters.TimeScale /= DilationFactor;
+	Parameters.TimeScale.ScaleBy(1.f / DilationFactor);
 }
 
 EMovieSceneChannelProxyType UTemplateSequenceSection::CacheChannelProxy()
diff --git a/Engine/Plugins/Runtime/ModelViewViewModel/Source/ModelViewViewModelBlueprint/Private/MVVMViewBlueprintCompiler.cpp b/Engine/Plugins/Runtime/ModelViewViewModel/Source/ModelViewViewModelBlueprint/Private/MVVMViewBlueprintCompiler.cpp
--- a/Engine/Plugins/Runtime/ModelViewViewModel/Source/ModelViewViewModelBlueprint/Private/MVVMViewBlueprintCompiler.cpp
+++ b/Engine/Plugins/Runtime/ModelViewViewModel/Source/ModelViewViewModelBlueprint/Private/MVVMViewBlueprintCompiler.cpp
@@ -24,6 +24,7 @@
 #include "UObject/UnrealType.h"
 #include "View/MVVMViewClass.h"
 #include "View/MVVMViewModelContextResolver.h"
+#include "Algo/AnyOf.h"
 #include "WidgetBlueprintEditorUtils.h"
 
 #define LOCTEXT_NAMESPACE "MVVMViewBlueprintCompiler"
diff --git a/Engine/Plugins/Runtime/USDCore/Source/USDUtilities/Private/USDPrimConversion.cpp b/Engine/Plugins/Runtime/USDCore/Source/USDUtilities/Private/USDPrimConversion.cpp
--- a/Engine/Plugins/Runtime/USDCore/Source/USDUtilities/Private/USDPrimConversion.cpp
+++ b/Engine/Plugins/Runtime/USDCore/Source/USDUtilities/Private/USDPrimConversion.cpp
@@ -2321,6 +2321,8 @@ bool UnrealToUsd::ConvertAudioSection(
 	pxr::UsdAttribute EndAttr = SpatialAudio.CreateEndTimeAttr();
 	if (StartAttr && EndAttr)
 	{
+		FMovieSceneInverseSequenceTransform InverseSequenceTransform = SequenceTransform.Inverse();
+
 		UsdUtils::NotifyIfOverriddenOpinion(StartAttr);
 		UsdUtils::NotifyIfOverriddenOpinion(EndAttr);
 
@@ -2333,8 +2335,8 @@ bool UnrealToUsd::ConvertAudioSection(
 
 		FFrameTime UsdStartTickTime = FFrameRate::Snap(StartTick, Resolution, DisplayRate).FloorToFrame();
 		FFrameTime UsdEndTickTime = FFrameRate::Snap(EndTick, Resolution, DisplayRate).FloorToFrame();
-		UsdStartTickTime *= SequenceTransform.InverseNoLooping();
-		UsdEndTickTime *= SequenceTransform.InverseNoLooping();
+		UsdStartTickTime = InverseSequenceTransform.TryTransformTime(UsdStartTickTime).Get(UsdStartTickTime);
+		UsdEndTickTime = InverseSequenceTransform.TryTransformTime(UsdEndTickTime).Get(UsdEndTickTime);
 		double UsdStartTimeCode = FFrameRate::TransformTime(UsdStartTickTime, Resolution, StageFrameRate).AsDecimal();
 		double UsdEndTimeCode = FFrameRate::TransformTime(UsdEndTickTime, Resolution, StageFrameRate).AsDecimal();
 
@@ -2478,6 +2480,8 @@ bool UnrealToUsd::ConvertBoolTrack(
 	const double StageTimeCodesPerSecond = Stage.GetTimeCodesPerSecond();
 	const FFrameRate StageFrameRate(StageTimeCodesPerSecond, 1);
 
+	FMovieSceneInverseSequenceTransform SequenceToRootTransform = SequenceTransform.Inverse();
+
 	auto EvaluateChannel = [&Resolution,
 							&DisplayRate](const FMovieSceneBoolChannel& Channel, bool InDefaultValue) -> TArray<TPair<FFrameNumber, bool>>
 	{
@@ -2508,13 +2512,15 @@ bool UnrealToUsd::ConvertBoolTrack(
 		{
 			for (const TPair<FFrameNumber, bool>& Pair : EvaluateChannel(BoolSection->GetChannel(), false))
 			{
-				FFrameTime TransformedBakedKeyTime{Pair.Key};
-				TransformedBakedKeyTime *= SequenceTransform.InverseNoLooping();
-				FFrameTime UsdFrameTime = FFrameRate::TransformTime(TransformedBakedKeyTime, Resolution, StageFrameRate);
-
-				bool UEValue = Pair.Value;
+				TOptional<FFrameTime> RootTime = SequenceToRootTransform.TryTransformTime(Pair.Key);
+				if (RootTime)
+				{
+					// @todo: Resolution here is actually the local tick res - this is incorrect and should be converted to use the root resolution
+					FFrameTime UsdFrameTime = FFrameRate::TransformTime(RootTime.GetValue(), Resolution, StageFrameRate);
+					bool UEValue = Pair.Value;
 
-				WriterFunc(UEValue, UsdFrameTime.AsDecimal());
+					WriterFunc(UEValue, UsdFrameTime.AsDecimal());
+				}
 			}
 		}
 	}
@@ -2602,15 +2608,19 @@ bool UnrealToUsd::ConvertFloatChannel(
 		return Values;
 	};
 
+	FMovieSceneInverseSequenceTransform SequenceToRootTransform = SequenceTransform.Inverse();
+
 	for (const TPair<FFrameNumber, float>& Pair : EvaluateChannel(MovieSceneChannel, 0.0f))
 	{
-		FFrameTime TransformedBakedKeyTime{Pair.Key};
-		TransformedBakedKeyTime *= SequenceTransform.InverseNoLooping();
-		FFrameTime UsdFrameTime = FFrameRate::TransformTime(TransformedBakedKeyTime, Resolution, StageFrameRate);
-
-		float UEValue = Pair.Value;
+		TOptional<FFrameTime> RootTime = SequenceToRootTransform.TryTransformTime(Pair.Key);
+		if (RootTime)
+		{
+			// @todo: Resolution here is actually the local tick res - this is incorrect and should be converted to use the root resolution
+			FFrameTime UsdFrameTime = FFrameRate::TransformTime(RootTime.GetValue(), Resolution, StageFrameRate);
+			float UEValue = Pair.Value;
 
-		WriterFunc(UEValue, UsdFrameTime.AsDecimal());
+			WriterFunc(UEValue, UsdFrameTime.AsDecimal());
+		}
 	}
 	return true;
 }
@@ -2710,6 +2720,8 @@ bool UnrealToUsd::ConvertColorTrack(
 		}
 	};
 
+	FMovieSceneInverseSequenceTransform SequenceToRootTransform = SequenceTransform.Inverse();
+
 	for (UMovieSceneSection* Section : MovieSceneTrack.GetAllSections())
 	{
 		if (UMovieSceneColorSection* ColorSection = Cast<UMovieSceneColorSection>(Section))
@@ -2733,6 +2745,12 @@ bool UnrealToUsd::ConvertColorTrack(
 			// Sample all channels at the union of bake times, construct the value and write it out
 			for (const FFrameNumber UntransformedBakeTime : BakeTimeUnion)
 			{
+				TOptional<FFrameTime> RootTime = SequenceToRootTransform.TryTransformTime(UntransformedBakeTime);
+				if (!RootTime)
+				{
+					continue;
+				}
+
 				float RedValue = 0.0f;
 				float GreenValue = 0.0f;
 				float BlueValue = 0.0f;
@@ -2745,9 +2763,7 @@ bool UnrealToUsd::ConvertColorTrack(
 
 				FLinearColor Color{RedValue, GreenValue, BlueValue, AlphaValue};
 
-				FFrameTime TransformedBakedKeyTime{UntransformedBakeTime};
-				TransformedBakedKeyTime *= SequenceTransform.InverseNoLooping();
-				FFrameTime UsdFrameTime = FFrameRate::TransformTime(TransformedBakedKeyTime, Resolution, StageFrameRate);
+				FFrameTime UsdFrameTime = FFrameRate::TransformTime(RootTime.GetValue(), Resolution, StageFrameRate);
 
 				WriterFunc(Color, UsdFrameTime.AsDecimal());
 			}
@@ -2869,11 +2885,19 @@ bool UnrealToUsd::ConvertBoundsVectorTracks(
 	TArray<FFrameNumber> BakeTimeUnion = AllBakeTimes.Array();
 	BakeTimeUnion.Sort();
 
+	FMovieSceneInverseSequenceTransform SequenceToRootTransform = SequenceTransform.Inverse();
+
 	// Sample all channels at the union of bake times, construct the value and write it out
 	// This could be done more efficently, but in the general case we're only going to have one
 	// section per track anyway so it shouldn't matter much
 	for (const FFrameNumber UntransformedBakeTime : BakeTimeUnion)
 	{
+		TOptional<FFrameTime> RootTime = SequenceToRootTransform.TryTransformTime(UntransformedBakeTime);
+		if (!RootTime)
+		{
+			continue;
+		}
+
 		FVector MinValue{0};
 		for (const UMovieSceneSection* Section : MinSections)
 		{
@@ -2906,9 +2930,7 @@ bool UnrealToUsd::ConvertBoundsVectorTracks(
 			}
 		}
 
-		FFrameTime TransformedBakedKeyTime{UntransformedBakeTime};
-		TransformedBakedKeyTime *= SequenceTransform.InverseNoLooping();
-		FFrameTime UsdFrameTime = FFrameRate::TransformTime(TransformedBakedKeyTime, Resolution, StageFrameRate);
+		FFrameTime UsdFrameTime = FFrameRate::TransformTime(RootTime.GetValue(), Resolution, StageFrameRate);
 
 		WriterFunc(MinValue, MaxValue, UsdFrameTime.AsDecimal());
 	}
@@ -2998,6 +3020,8 @@ bool UnrealToUsd::Convert3DTransformTrack(
 		return BakeTimes;
 	};
 
+	FMovieSceneInverseSequenceTransform SequenceToRootTransform = SequenceTransform.Inverse();
+
 	for (UMovieSceneSection* Section : MovieSceneTrack.GetAllSections())
 	{
 		if (UMovieScene3DTransformSection* TransformSection = Cast<UMovieScene3DTransformSection>(Section))
@@ -3056,6 +3080,12 @@ bool UnrealToUsd::Convert3DTransformTrack(
 			// Sample all channels at the union of bake times, construct the value and write it out
 			for (const FFrameNumber UntransformedBakeTime : BakeTimeUnion)
 			{
+				TOptional<FFrameTime> RootTime = SequenceToRootTransform.TryTransformTime(UntransformedBakeTime);
+				if (!RootTime)
+				{
+					continue;
+				}
+
 				double LocX = 0.0f;
 				double LocY = 0.0f;
 				double LocZ = 0.0f;
@@ -3115,9 +3145,7 @@ bool UnrealToUsd::Convert3DTransformTrack(
 				FVector Scale{(float)ScaleX, (float)ScaleY, (float)ScaleZ};
 				FTransform Transform{Rotation, Location, Scale};
 
-				FFrameTime TransformedBakedKeyTime{UntransformedBakeTime};
-				TransformedBakedKeyTime *= SequenceTransform.InverseNoLooping();
-				FFrameTime UsdFrameTime = FFrameRate::TransformTime(TransformedBakedKeyTime, Resolution, StageFrameRate);
+				FFrameTime UsdFrameTime = FFrameRate::TransformTime(RootTime.GetValue(), Resolution, StageFrameRate);
 
 				WriterFunc(Transform, UsdFrameTime.AsDecimal());
 			}
@@ -5154,10 +5182,12 @@ bool UnrealToUsd::ConvertXformable(
 	const double StageTimeCodesPerSecond = UsdPrim.GetStage()->GetTimeCodesPerSecond();
 	const FFrameRate StageFrameRate(StageTimeCodesPerSecond, 1);
 
+	FMovieSceneInverseSequenceTransform SequenceToRootTransform = SequenceTransform.Inverse();
+
 	auto EvaluateChannel = [&PlaybackRange,
 							&Resolution,
 							&DisplayRate,
-							&SequenceTransform](const FMovieSceneDoubleChannel* Channel, double DefaultValue) -> TArray<TPair<FFrameNumber, float>>
+							&SequenceToRootTransform](const FMovieSceneDoubleChannel* Channel, double DefaultValue) -> TArray<TPair<FFrameNumber, float>>
 	{
 		TArray<TPair<FFrameNumber, float>> Values;
 
@@ -5171,6 +5201,12 @@ bool UnrealToUsd::ConvertXformable(
 			{
 				FFrameNumber KeyTime = FFrameRate::Snap(EvalTime, Resolution, DisplayRate).FloorToFrame();
 
+				TOptional<FFrameTime> RootTime = SequenceToRootTransform.TryTransformTime(KeyTime);
+				if (!RootTime)
+				{
+					continue;
+				}
+
 				double Result = DefaultValue;
 				if (Channel)
 				{
@@ -5178,10 +5214,7 @@ bool UnrealToUsd::ConvertXformable(
 					Channel->Evaluate(KeyTime, Result);
 				}
 
-				FFrameTime GlobalEvalTime(KeyTime);
-				GlobalEvalTime *= SequenceTransform.InverseNoLooping();
-
-				Values.Emplace(GlobalEvalTime.GetFrame(), Result);
+				Values.Emplace(RootTime->GetFrame(), Result);
 			}
 		}
 
diff --git a/Engine/Plugins/Runtime/USDCore/Source/USDUtilities/Private/USDSkeletalDataConversion.cpp b/Engine/Plugins/Runtime/USDCore/Source/USDUtilities/Private/USDSkeletalDataConversion.cpp
--- a/Engine/Plugins/Runtime/USDCore/Source/USDUtilities/Private/USDSkeletalDataConversion.cpp
+++ b/Engine/Plugins/Runtime/USDCore/Source/USDUtilities/Private/USDSkeletalDataConversion.cpp
@@ -3972,7 +3972,7 @@ bool UnrealToUsd::ConvertAnimSequence(UAnimSequence* AnimSequence, pxr::UsdPrim&
 
 bool UnrealToUsd::ConvertControlRigSection(
 	UMovieSceneControlRigParameterSection* InSection,
-	const FMovieSceneSequenceTransform& InTransform,
+	const FMovieSceneInverseSequenceTransform& InTransform,
 	UMovieScene* InMovieScene,
 	IMovieScenePlayer* InPlayer,
 	const FReferenceSkeleton& InRefSkeleton,
@@ -4292,12 +4292,16 @@ bool UnrealToUsd::ConvertControlRigSection(
 	FFrameTime TickIncr = FFrameRate::TransformTime(1, DisplayRate, TickResolution);
 	for (FFrameTime FrameTickTime = StartInclTickFrame; FrameTickTime <= EndInclTickFrame; FrameTickTime += TickIncr)
 	{
-		FFrameTime TransformedFrameTickTime = FrameTickTime * InTransform;
+		TOptional<FFrameTime> TransformedFrameTickTime = InTransform.TryTransformTime(FrameTickTime);
+		if (!TransformedFrameTickTime)
+		{
+			continue;
+		}
 
-		double UsdTimeCode = FFrameRate::TransformTime(TransformedFrameTickTime, TickResolution, StageFrameRate).AsDecimal();
+		double UsdTimeCode = FFrameRate::TransformTime(TransformedFrameTickTime.GetValue(), TickResolution, StageFrameRate).AsDecimal();
 
 		FMovieSceneContext Context = FMovieSceneContext(
-										 FMovieSceneEvaluationRange(TransformedFrameTickTime, TickResolution),
+										 FMovieSceneEvaluationRange(TransformedFrameTickTime.GetValue(), TickResolution),
 										 InPlayer->GetPlaybackStatus()
 		)
 										 .SetHasJumped(true);
diff --git a/Engine/Plugins/Runtime/USDCore/Source/USDUtilities/Public/USDSkeletalDataConversion.h b/Engine/Plugins/Runtime/USDCore/Source/USDUtilities/Public/USDSkeletalDataConversion.h
--- a/Engine/Plugins/Runtime/USDCore/Source/USDUtilities/Public/USDSkeletalDataConversion.h
+++ b/Engine/Plugins/Runtime/USDCore/Source/USDUtilities/Public/USDSkeletalDataConversion.h
@@ -39,7 +39,7 @@ class UMovieScene;
 class UMovieSceneControlRigParameterSection;
 class USkeletalMesh;
 struct FUsdStageInfo;
-struct FMovieSceneSequenceTransform;
+struct FMovieSceneInverseSequenceTransform;
 
 namespace SkeletalMeshImportData
 {
@@ -436,7 +436,7 @@ namespace UnrealToUsd
 	 */
 	USDUTILITIES_API bool ConvertControlRigSection(
 		UMovieSceneControlRigParameterSection* InSection,
-		const FMovieSceneSequenceTransform& InTransform,
+		const FMovieSceneInverseSequenceTransform& InTransform,
 		UMovieScene* InMovieScene,
 		IMovieScenePlayer* InPlayer,
 		const FReferenceSkeleton& InRefSkeleton,
diff --git a/Engine/Plugins/VirtualProduction/DMX/DMXEngine/Source/DMXRuntime/Private/Sequencer/MovieSceneDMXLibrarySection.cpp b/Engine/Plugins/VirtualProduction/DMX/DMXEngine/Source/DMXRuntime/Private/Sequencer/MovieSceneDMXLibrarySection.cpp
--- a/Engine/Plugins/VirtualProduction/DMX/DMXEngine/Source/DMXRuntime/Private/Sequencer/MovieSceneDMXLibrarySection.cpp
+++ b/Engine/Plugins/VirtualProduction/DMX/DMXEngine/Source/DMXRuntime/Private/Sequencer/MovieSceneDMXLibrarySection.cpp
@@ -19,6 +19,7 @@
 #include "Library/DMXEntityFixtureType.h"
 #include "Library/DMXLibrary.h"
 #include "UObject/Package.h"
+#include "Algo/Find.h"
 
 DECLARE_LOG_CATEGORY_CLASS(MovieSceneDMXLibrarySectionLog, Log, All);
 
diff --git a/Engine/Plugins/VirtualProduction/SequencerPlaylists/Source/SequencerPlaylists/Private/SequencerPlaylistItem_Sequence.cpp b/Engine/Plugins/VirtualProduction/SequencerPlaylists/Source/SequencerPlaylists/Private/SequencerPlaylistItem_Sequence.cpp
--- a/Engine/Plugins/VirtualProduction/SequencerPlaylists/Source/SequencerPlaylists/Private/SequencerPlaylistItem_Sequence.cpp
+++ b/Engine/Plugins/VirtualProduction/SequencerPlaylists/Source/SequencerPlaylists/Private/SequencerPlaylistItem_Sequence.cpp
@@ -46,24 +46,6 @@ namespace UE::Private::SequencerPlaylistItem_Sequence
 
 		return nullptr;
 	}
-
-
-	TPair<TOptional<FFrameTime>, FMovieSceneWarpCounter>
-	GetInnerTimeAndWarp(FQualifiedFrameTime QualifiedTime, UMovieSceneSubSection* SubSection)
-	{
-		check(SubSection);
-
-		TPair<TOptional<FFrameTime>, FMovieSceneWarpCounter> Result;
-
-		if (IsTimeWithinSection(QualifiedTime, SubSection))
-		{
-			FFrameTime InnerTime;
-			SubSection->OuterToInnerTransform().TransformTime(QualifiedTime.Time, InnerTime, Result.Get<1>());
-			Result.Get<0>() = InnerTime;
-		}
-
-		return MoveTemp(Result);
-	}
 }
 
 
@@ -290,7 +272,7 @@ bool FSequencerPlaylistItemPlayer_Sequence::InternalPlay(
 
 	if (PlayParams.Direction == ESequencerPlaylistPlaybackDirection::Reverse)
 	{
-		WorkingSubSection->Parameters.TimeScale *= -1.0f;
+		WorkingSubSection->Parameters.TimeScale.ScaleBy(-1.0f);
 		WorkingSubSection->Parameters.StartFrameOffset =
 			SequencePlayRange.Size<FFrameNumber>() - WorkingSubSection->Parameters.StartFrameOffset;
 	}
@@ -336,12 +318,9 @@ bool FSequencerPlaylistItemPlayer_Sequence::InternalPause(USequencerPlaylistItem
 
 			const FQualifiedFrameTime GlobalTime = Sequencer->GetGlobalTime();
 			UMovieSceneSubSection* CurrentPlayingSection = IsTimeWithinAnySection(GlobalTime, ItemState.WeakPlaySections);
-			if (CurrentPlayingSection)
+			if (CurrentPlayingSection && IsTimeWithinSection(GlobalTime, CurrentPlayingSection))
 			{
-				TPair<TOptional<FFrameTime>, FMovieSceneWarpCounter> InnerTimeAndWarp =
-					GetInnerTimeAndWarp(GlobalTime, CurrentPlayingSection);
-
-				HoldParams.StartFrameOffset_SceneTicks = InnerTimeAndWarp.Get<0>().GetValue();
+				HoldParams.StartFrameOffset_SceneTicks = CurrentPlayingSection->OuterToInnerTransform().TransformTime(GlobalTime.Time);
 			}
 
 			bSequenceWasModified |= Stop(Item);
@@ -634,10 +613,10 @@ bool FSequencerPlaylistItemPlayer_Sequence::EndSection(UMovieSceneSection* Secti
 		{
 			if (SubSection->Parameters.bCanLoop)
 			{
-				// Calc whether we've looped; if not, set bCanLoop back to false.
-				TPair<TOptional<FFrameTime>, FMovieSceneWarpCounter> InnerTimeAndWarp =
-					GetInnerTimeAndWarp(GlobalTime, SubSection);
-				if (InnerTimeAndWarp.Get<1>().LastWarpCount() == 0)
+				FMovieSceneSubSequenceData SubData(*SubSection);
+
+				TOptional<FFrameTime> FirstLoopEndTime = SubData.OuterToInnerTransform.Inverse().TryTransformTime(SubData.PlayRange.Value.GetUpperBoundValue());
+				if (FirstLoopEndTime && GlobalTime.Time <= FirstLoopEndTime.GetValue())
 				{
 					SubSection->Parameters.bCanLoop = false;
 				}
diff --git a/Engine/Source/Editor/CurveEditor/Private/CurveEditor.cpp b/Engine/Source/Editor/CurveEditor/Private/CurveEditor.cpp
--- a/Engine/Source/Editor/CurveEditor/Private/CurveEditor.cpp
+++ b/Engine/Source/Editor/CurveEditor/Private/CurveEditor.cpp
@@ -9,6 +9,7 @@
 #include "CurveEditorCopyBuffer.h"
 #include "CurveEditorSettings.h"
 #include "CurveEditorSnapMetrics.h"
+#include "CurveEditorAxis.h"
 #include "CurveModel.h"
 #include "Curves/KeyHandle.h"
 #include "Curves/RichCurve.h"
@@ -180,12 +181,42 @@ FCurveEditorToolID FCurveEditor::AddTool(TUniquePtr<ICurveEditorToolExtension>&&
 	return NewID;
 }
 
+void FCurveEditor::AddAxis(const FName& InIdentifier, TSharedPtr<FCurveEditorAxis> InAxis)
+{
+	// Allow overwrites
+	CustomAxes.Add(InIdentifier, InAxis);
+}
+
+TSharedPtr<FCurveEditorAxis> FCurveEditor::FindAxis(const FName& InIdentifier) const
+{
+	return CustomAxes.FindRef(InIdentifier);
+}
+
+void FCurveEditor::RemoveAxis(const FName& InIdentifier)
+{
+	CustomAxes.Remove(InIdentifier);
+}
+
+void FCurveEditor::ClearAxes()
+{
+	CustomAxes.Empty();
+}
+
 FCurveModelID FCurveEditor::AddCurve(TUniquePtr<FCurveModel>&& InCurve)
 {
 	FCurveModelID NewID = FCurveModelID::Unique();
 	FCurveModel *Curve = InCurve.Get();
 
 	CurveData.Add(NewID, MoveTemp(InCurve));
+
+	// Add child curves
+	TArray<TUniquePtr<FCurveModel>> ChildCurvesArray;
+	Curve->MakeChildCurves(ChildCurvesArray);
+	for (TUniquePtr<FCurveModel>& Child : ChildCurvesArray)
+	{
+		ChildCurves.Add(NewID, AddCurve(MoveTemp(Child)));
+	}
+
 	++ActiveCurvesSerialNumber;
 	if (IsBroadcasting())
 	{
@@ -204,21 +235,11 @@ void FCurveEditor::BroadcastCurveChanged(FCurveModel* InCurve)
 
 FCurveModelID FCurveEditor::AddCurveForTreeItem(TUniquePtr<FCurveModel>&& InCurve, FCurveEditorTreeItemID TreeItemID)
 {
-	FCurveModelID NewID = FCurveModelID::Unique();
-	FCurveModel *Curve = InCurve.Get();
-
-	if(IsBroadcasting())
-	{
-		OnCurveArrayChanged.Broadcast(InCurve.Get(), true, this);
-	}
-
-	CurveData.Add(NewID, MoveTemp(InCurve));
+	FCurveModelID NewID = AddCurve(MoveTemp(InCurve));
 	TreeIDByCurveID.Add(NewID, TreeItemID);
-
-	++ActiveCurvesSerialNumber;
-
 	return NewID;
 }
+
 void FCurveEditor::ResetMinMaxes()
 {
 	TSharedPtr<SCurveEditorPanel> Panel = WeakPanel.Pin();
@@ -229,6 +250,12 @@ void FCurveEditor::ResetMinMaxes()
 }
 void FCurveEditor::RemoveCurve(FCurveModelID InCurveID)
 {
+	for (auto ChildID = ChildCurves.CreateConstKeyIterator(InCurveID); ChildID; ++ChildID)
+	{
+		RemoveCurve(ChildID.Value());
+	}
+	ChildCurves.Remove(InCurveID);
+
 	TSharedPtr<SCurveEditorPanel> Panel = WeakPanel.Pin();
 	if (Panel.IsValid())
 	{
@@ -247,7 +274,6 @@ void FCurveEditor::RemoveCurve(FCurveModelID InCurveID)
 
 
 	++ActiveCurvesSerialNumber;
-
 }
 
 void FCurveEditor::RemoveAllCurves()
@@ -264,6 +290,7 @@ void FCurveEditor::RemoveAllCurves()
 	CurveData.Empty();
 	Selection.Clear();
 	PinnedCurves.Empty();
+	ChildCurves.Empty();
 
 	++ActiveCurvesSerialNumber;
 }
@@ -330,7 +357,7 @@ FCurveEditorTreeItemID FCurveEditor::GetTreeIDFromCurveID(FCurveModelID CurveID)
 {
 	if (TreeIDByCurveID.Contains(CurveID))
 	{
-		return TreeIDByCurveID[CurveID];	
+		return TreeIDByCurveID[CurveID];
 	}
 
 	return FCurveEditorTreeItemID();
@@ -591,9 +618,53 @@ void FCurveEditor::ZoomToFitInternal(EAxisList::Type Axes, const TMap<FCurveMode
 {
 	TArray<FKeyPosition> KeyPositionsScratch;
 
-	double InputMin = TNumericLimits<double>::Max(), InputMax = TNumericLimits<double>::Lowest();
+	TMap<TTuple<TSharedRef<SCurveEditorView>, FCurveEditorViewAxisID>, TTuple<double, double>> ViewAndAxisToInputBounds;
+	TMap<TTuple<TSharedRef<SCurveEditorView>, FCurveEditorViewAxisID>, TTuple<double, double>> ViewAndAxisToOutputBounds;
+
+	auto TrackHorizontalBoundsForView = [&ViewAndAxisToInputBounds, Axes](const TSharedRef<SCurveEditorView>& View, FCurveModelID InCurveID, double InputMin, double InputMax)
+	{
+		if (Axes & EAxisList::X)
+		{
+			FCurveEditorViewAxisID HorizontalAxis = View->GetAxisForCurve(InCurveID, ECurveEditorAxisOrientation::Horizontal);
+			if (HorizontalAxis)  // Only track horizontal axis zoom for custom axes since every view is implicitly linked to the global curve editor bounds
+			{
+				TTuple<double, double>* ViewBounds = ViewAndAxisToInputBounds.Find(MakeTuple(View, HorizontalAxis));
+				if (ViewBounds)
+				{
+					ViewBounds->Get<0>() = FMath::Min(ViewBounds->Get<0>(), InputMin);
+					ViewBounds->Get<1>() = FMath::Max(ViewBounds->Get<1>(), InputMax);
+				}
+				else
+				{
+					ViewAndAxisToInputBounds.Add(MakeTuple(View, HorizontalAxis), MakeTuple(InputMin, InputMax));
+				}
+			}
+		}
+	};
 
-	TMap<TSharedRef<SCurveEditorView>, TTuple<double, double>> ViewToOutputBounds;
+	auto TrackVerticalBoundsForView = [&ViewAndAxisToOutputBounds, Axes](const TSharedRef<SCurveEditorView>& View, FCurveModelID InCurveID, double OutputMin, double OutputMax)
+	{
+		if (Axes & EAxisList::Y)
+		{
+			FCurveEditorViewAxisID VerticalAxis = View->GetAxisForCurve(InCurveID, ECurveEditorAxisOrientation::Vertical);
+
+			TTuple<double, double>* ViewBounds = ViewAndAxisToOutputBounds.Find(MakeTuple(View, VerticalAxis));
+			if (ViewBounds)
+			{
+				ViewBounds->Get<0>() = FMath::Min(ViewBounds->Get<0>(), OutputMin);
+				ViewBounds->Get<1>() = FMath::Max(ViewBounds->Get<1>(), OutputMax);
+			}
+			else
+			{
+				ViewAndAxisToOutputBounds.Add(MakeTuple(View, VerticalAxis), MakeTuple(OutputMin, OutputMax));
+			}
+		}
+	};
+
+	double AllInputMin = TNumericLimits<double>::Max(), AllInputMax = TNumericLimits<double>::Lowest();
+
+	TSharedPtr<SCurveEditorPanel> Panel = WeakPanel.Pin();
+	TSharedPtr<SCurveEditorView>  View  = WeakView.Pin();
 
 	for (const TTuple<FCurveModelID, FKeyHandleSet>& Pair : CurveKeySet)
 	{
@@ -604,6 +675,7 @@ void FCurveEditor::ZoomToFitInternal(EAxisList::Type Axes, const TMap<FCurveMode
 			continue;
 		}
 
+		double InputMin  = TNumericLimits<double>::Max(), InputMax  = TNumericLimits<double>::Lowest();
 		double OutputMin = TNumericLimits<double>::Max(), OutputMax = TNumericLimits<double>::Lowest();
 
 		int32 NumKeys = Pair.Value.AsArray().Num();
@@ -639,110 +711,110 @@ void FCurveEditor::ZoomToFitInternal(EAxisList::Type Axes, const TMap<FCurveMode
 			}
 		}
 
-		if (Axes & EAxisList::Y)
+		AllInputMin = FMath::Min(InputMin, AllInputMin);
+		AllInputMax = FMath::Max(InputMax, AllInputMax);
+
+		if (Panel)
 		{
-			TSharedPtr<SCurveEditorPanel> Panel = WeakPanel.Pin();
-			TSharedPtr<SCurveEditorView> View = WeakView.Pin();
-			if (Panel.IsValid())
+			// Store the min max for each view
+			for (auto ViewIt = Panel->FindViews(CurveID); ViewIt; ++ViewIt)
 			{
-				// Store the min max for each view
-				for (auto ViewIt = Panel->FindViews(CurveID); ViewIt; ++ViewIt)
-				{
-					TTuple<double, double>* ViewBounds = ViewToOutputBounds.Find(ViewIt.Value());
-					if (ViewBounds)
-					{
-						ViewBounds->Get<0>() = FMath::Min(ViewBounds->Get<0>(), OutputMin);
-						ViewBounds->Get<1>() = FMath::Max(ViewBounds->Get<1>(), OutputMax);
-					}
-					else
-					{
-						ViewToOutputBounds.Add(ViewIt.Value(), MakeTuple(OutputMin, OutputMax));
-					}
-				}
-			}
-			else if(View.IsValid())
-			{
-				TTuple<double, double>* ViewBounds = ViewToOutputBounds.Find(View.ToSharedRef());
-				if (ViewBounds)
-				{
-					ViewBounds->Get<0>() = FMath::Min(ViewBounds->Get<0>(), OutputMin);
-					ViewBounds->Get<1>() = FMath::Max(ViewBounds->Get<1>(), OutputMax);
-				}
-				else
-				{
-					ViewToOutputBounds.Add(View.ToSharedRef(), MakeTuple(OutputMin, OutputMax));
-				}
+				TrackHorizontalBoundsForView(ViewIt.Value(), CurveID, InputMin, InputMax);
+				TrackVerticalBoundsForView(ViewIt.Value(), CurveID, OutputMin, OutputMax);
 			}
 		}
+		else if(View.IsValid())
+		{
+			TrackHorizontalBoundsForView(View.ToSharedRef(), CurveID, InputMin, InputMax);
+			TrackVerticalBoundsForView(View.ToSharedRef(), CurveID, OutputMin, OutputMax);
+		}
 	}
 
-	if (Axes & EAxisList::X && InputMin != TNumericLimits<double>::Max() && InputMax != TNumericLimits<double>::Lowest())
+	auto AdjustHorizontalBounds = [this, Panel, View](TSharedPtr<SCurveEditorView> InView, double CurrentInputMin, double CurrentInputMax, double& NewInputMin, double& NewInputMax)
 	{
 		// If zooming to the same (or invalid) min/max, keep the same zoom scale and center within the timeline
-		if (InputMin >= InputMax)
+		if (NewInputMin >= NewInputMax)
 		{
-			double CurrentInputMin = 0.0, CurrentInputMax = 1.0;
-			Bounds->GetInputBounds(CurrentInputMin, CurrentInputMax);
-
-			const double HalfInputScale = (CurrentInputMax - CurrentInputMin)*0.5;
-			InputMin -= HalfInputScale;
-			InputMax += HalfInputScale;
+			const double HalfInputScale = (CurrentInputMax - CurrentInputMin) * 0.5;
+			NewInputMin -= HalfInputScale;
+			NewInputMax += HalfInputScale;
 		}
 		else
 		{
-			TSharedPtr<SCurveEditorPanel> Panel = WeakPanel.Pin();
-			TSharedPtr<SCurveEditorView> View = WeakView.Pin();
-
-			double PanelWidth = 0;
-			if (Panel.IsValid())
+			double PanelHeight = 0;
+			if (Panel)
 			{
-				PanelWidth = WeakPanel.Pin()->GetViewContainerGeometry().GetLocalSize().X;
+				PanelHeight = Panel->GetViewContainerGeometry().GetLocalSize().Y;
 			}
-			else if (View.IsValid())
+			else
 			{
-				PanelWidth = View->GetViewSpace().GetPhysicalWidth();
+				PanelHeight = InView->GetViewSpace().GetPhysicalHeight();
 			}
-			
-			double InputPercentage = PanelWidth != 0 ? FMath::Min(Settings->GetFrameInputPadding() / PanelWidth, 0.5) : 0.1; // Cannot pad more than half the width
 
-			const double MinInputZoom = InputSnapEnabledAttribute.Get() ? InputSnapRateAttribute.Get().AsInterval() : 0.00001;
-			const double InputPadding = FMath::Max((InputMax - InputMin) * InputPercentage, MinInputZoom);
-			InputMax = FMath::Max(InputMin + MinInputZoom, InputMax);
+			double InputPercentage = PanelHeight != 0 ? FMath::Min(Settings->GetFrameInputPadding() / PanelHeight, 0.5) : 0.1; // Cannot pad more than half the height
 
-			InputMin -= InputPadding;
-			InputMax += InputPadding;
+			constexpr double MinInputZoom = 0.00001;
+			const double InputPadding = FMath::Max((NewInputMax - NewInputMin) * InputPercentage, MinInputZoom);
+
+			NewInputMin -= InputPadding;
+			NewInputMax = FMath::Max(NewInputMin + MinInputZoom, NewInputMax) + InputPadding;
 		}
+	};
+
+	// Perform per-view input zoom for custom axes
+	for (const TPair<TTuple<TSharedRef<SCurveEditorView>, FCurveEditorViewAxisID>, TTuple<double, double>>& ViewAndAxisToBounds : ViewAndAxisToInputBounds)
+	{
+		FCurveEditorViewAxisID       AxisID   = ViewAndAxisToBounds.Key.Value;
+		TSharedRef<SCurveEditorView> AxisView = ViewAndAxisToBounds.Key.Key;
 
-		Bounds->SetInputBounds(InputMin, InputMax);
+		check(AxisID);
+
+		FCurveEditorScreenSpaceH AxisSpace = AxisView->GetHorizontalAxisSpace(AxisID);
+
+		double InputMin = ViewAndAxisToBounds.Value.Get<0>();
+		double InputMax = ViewAndAxisToBounds.Value.Get<1>();
+
+		AdjustHorizontalBounds(AxisView, AxisSpace.GetInputMin(), AxisSpace.GetInputMax(), InputMin, InputMax);
+
+		AxisView->FrameHorizontal(InputMin, InputMax, AxisID);
+	}
+
+	if (Axes & EAxisList::X && AllInputMin != TNumericLimits<double>::Max() && AllInputMax != TNumericLimits<double>::Lowest())
+	{
+		double CurrentInputMin = 0.0, CurrentInputMax = 1.0;
+		Bounds->GetInputBounds(CurrentInputMin, CurrentInputMax);
+
+		AdjustHorizontalBounds(View, CurrentInputMin, CurrentInputMax, AllInputMin, AllInputMax);
+
+		Bounds->SetInputBounds(AllInputMin, AllInputMax);
 	}
 
 	// Perform per-view output zoom for any computed ranges
-	for (const TTuple<TSharedRef<SCurveEditorView>, TTuple<double, double>>& ViewAndBounds : ViewToOutputBounds)
+	for (const TPair<TTuple<TSharedRef<SCurveEditorView>, FCurveEditorViewAxisID>, TTuple<double, double>>& ViewAndAxisToBounds : ViewAndAxisToOutputBounds)
 	{
-		TSharedRef<SCurveEditorView> View = ViewAndBounds.Key;
+		FCurveEditorViewAxisID       AxisID   = ViewAndAxisToBounds.Key.Value;
+		TSharedRef<SCurveEditorView> AxisView = ViewAndAxisToBounds.Key.Key;
 
-		double OutputMin = ViewAndBounds.Value.Get<0>();
-		double OutputMax = ViewAndBounds.Value.Get<1>();
+		double OutputMin = ViewAndAxisToBounds.Value.Get<0>();
+		double OutputMax = ViewAndAxisToBounds.Value.Get<1>();
 
 		// If zooming to the same (or invalid) min/max, keep the same zoom scale and center within the timeline
 		if (OutputMin >= OutputMax)
 		{
-			const double HalfOutputScale = (View->GetOutputMax() - View->GetOutputMin()) * 0.5;
+			const double HalfOutputScale = (AxisView->GetOutputMax() - AxisView->GetOutputMin()) * 0.5;
 			OutputMin -= HalfOutputScale;
 			OutputMax += HalfOutputScale;
 		}
 		else
 		{
-			TSharedPtr<SCurveEditorPanel> Panel = WeakPanel.Pin();
-
 			double PanelHeight = 0;
-			if (Panel.IsValid())
+			if (Panel)
 			{
-				PanelHeight = WeakPanel.Pin()->GetViewContainerGeometry().GetLocalSize().Y;
+				PanelHeight = Panel->GetViewContainerGeometry().GetLocalSize().Y;
 			}
 			else
 			{
-				PanelHeight = View->GetViewSpace().GetPhysicalHeight();
+				PanelHeight = AxisView->GetViewSpace().GetPhysicalHeight();
 			}
 
 			double OutputPercentage = PanelHeight != 0 ? FMath::Min(Settings->GetFrameOutputPadding() / PanelHeight, 0.5) : 0.1; // Cannot pad more than half the height
@@ -753,7 +825,8 @@ void FCurveEditor::ZoomToFitInternal(EAxisList::Type Axes, const TMap<FCurveMode
 			OutputMin -= OutputPadding;
 			OutputMax = FMath::Max(OutputMin + MinOutputZoom, OutputMax) + OutputPadding;
 		}
-		View->FrameVertical(OutputMin, OutputMax);
+
+		AxisView->FrameVertical(OutputMin, OutputMax, AxisID);
 	}
 }
 
diff --git a/Engine/Source/Editor/CurveEditor/Private/CurveEditorHelpers.cpp b/Engine/Source/Editor/CurveEditor/Private/CurveEditorHelpers.cpp
--- a/Engine/Source/Editor/CurveEditor/Private/CurveEditorHelpers.cpp
+++ b/Engine/Source/Editor/CurveEditor/Private/CurveEditorHelpers.cpp
@@ -47,6 +47,41 @@ FVector2D GetVectorFromSlopeAndLength(float Slope, float Length)
 	return FVector2D(x, y);
 }
 
+void PopulateGridLineValues(float PhysicalSize, double ViewMin, double ViewMax, uint8 InMinorDivisions, TArray<double>& OutMajorGridLines, TArray<double>& OutMinorGridLines)
+{
+	const double PixelsPerValue = PhysicalSize / FMath::Max(ViewMax - ViewMin, 1e-10);
+	const double GridPixelSpacing = PhysicalSize / 5.0;
+
+	const double Order = FMath::Pow(10.0, FMath::FloorToInt(FMath::LogX(10.0, GridPixelSpacing / PixelsPerValue)));
+
+	static const int32 DesirableBases[]  = { 2, 5 };
+	static const int32 NumDesirableBases = UE_ARRAY_COUNT(DesirableBases);
+
+	const int32 Scale = FMath::RoundToInt(GridPixelSpacing / PixelsPerValue / Order);
+	int32 Base = DesirableBases[0];
+	for (int32 BaseIndex = 1; BaseIndex < NumDesirableBases; ++BaseIndex)
+	{
+		if (FMath::Abs(Scale - DesirableBases[BaseIndex]) < FMath::Abs(Scale - Base))
+		{
+			Base = DesirableBases[BaseIndex];
+		}
+	}
+
+	double MajorGridStep = FMath::Pow(static_cast<float>(Base), FMath::FloorToFloat(FMath::LogX(static_cast<float>(Base), static_cast<float>(Scale)))) * Order;
+
+	const double FirstMajorLine = FMath::FloorToDouble(ViewMin / MajorGridStep) * MajorGridStep;
+	const double LastMajorLine = FMath::CeilToDouble(ViewMax / MajorGridStep) * MajorGridStep;
+
+	for (double CurrentMajorLine = FirstMajorLine; CurrentMajorLine <= LastMajorLine; CurrentMajorLine += MajorGridStep)
+	{
+		OutMajorGridLines.Add(CurrentMajorLine);
+		for (int32 Step = 1; Step < InMinorDivisions; ++Step)
+		{
+			OutMinorGridLines.Add(CurrentMajorLine + Step * MajorGridStep / InMinorDivisions);
+		}
+	}
+}
+
 void ConstructYGridLines(const FCurveEditorScreenSpace& ViewSpace, uint8 InMinorDivisions, TArray<float>& OutMajorGridLines, TArray<float>& OutMinorGridLines, FText GridLineLabelFormatY, TArray<FText>* OutMajorGridLabels)
 {
 	const double GridPixelSpacing = ViewSpace.GetPhysicalHeight() / 5.0;
diff --git a/Engine/Source/Editor/CurveEditor/Private/DragOperations/CurveEditorDragOperation_MoveKeys.cpp b/Engine/Source/Editor/CurveEditor/Private/DragOperations/CurveEditorDragOperation_MoveKeys.cpp
--- a/Engine/Source/Editor/CurveEditor/Private/DragOperations/CurveEditorDragOperation_MoveKeys.cpp
+++ b/Engine/Source/Editor/CurveEditor/Private/DragOperations/CurveEditorDragOperation_MoveKeys.cpp
@@ -55,6 +55,8 @@ void FCurveEditorDragOperation_MoveKeys::OnBeginDrag(FVector2D InitialPosition,
 
 			KeyData.StartKeyPositions.SetNumZeroed(KeyData.Handles.Num());
 			Curve->GetKeyPositions(KeyData.Handles, KeyData.StartKeyPositions);
+
+			KeyData.InitialDragTransform = Curve->GetCurveTransform();
 			KeyData.LastDraggedKeyPositions = KeyData.StartKeyPositions;
 		}
 	}
@@ -92,6 +94,10 @@ void FCurveEditorDragOperation_MoveKeys::OnDrag(FVector2D InitialPosition, FVect
 
 		FCurveSnapMetrics SnapMetrics = CurveEditor->GetCurveSnapMetrics(KeyData.CurveID);
 
+		// Transform by the inverse curve transform to put the start positions in view space
+		FTransform2d CurveTransform               = Curve->GetCurveTransform();
+		FTransform2d InverseInitialCurveTransform = KeyData.InitialDragTransform.Inverse();
+
 		if (CardinalPoint.IsSet())
 		{
 			for (int KeyIndex = 0; KeyIndex < KeyData.StartKeyPositions.Num(); ++KeyIndex)
@@ -99,7 +105,7 @@ void FCurveEditorDragOperation_MoveKeys::OnDrag(FVector2D InitialPosition, FVect
 				FKeyHandle KeyHandle = KeyData.Handles[KeyIndex];
 				if (CardinalPoint->KeyHandle == KeyHandle)
 				{
-					FKeyPosition StartPosition = KeyData.StartKeyPositions[KeyIndex];
+					FKeyPosition StartPosition = KeyData.StartKeyPositions[KeyIndex].Transform(InverseInitialCurveTransform);
 
 					if (View->IsTimeSnapEnabled())
 					{
@@ -111,13 +117,16 @@ void FCurveEditorDragOperation_MoveKeys::OnDrag(FVector2D InitialPosition, FVect
 					{
 						DeltaOutput = SnapMetrics.SnapOutput(StartPosition.OutputValue + DeltaOutput) - StartPosition.OutputValue;
 					}
+
+					break;
 				}
 			}
 		}
-
 		for (FKeyPosition StartPosition : KeyData.StartKeyPositions)
 		{
-			StartPosition.InputValue  += DeltaInput;
+			StartPosition = StartPosition.Transform(InverseInitialCurveTransform);
+
+			StartPosition.InputValue += DeltaInput;
 			StartPosition.OutputValue += DeltaOutput;
 
 			StartPosition.InputValue = View->IsTimeSnapEnabled() ? SnapMetrics.SnapInputSeconds(StartPosition.InputValue) : StartPosition.InputValue;
@@ -127,8 +136,8 @@ void FCurveEditorDragOperation_MoveKeys::OnDrag(FVector2D InitialPosition, FVect
 			{
 				StartPosition.OutputValue = View->IsValueSnapEnabled() ? SnapMetrics.SnapOutput(StartPosition.OutputValue) : StartPosition.OutputValue;
 			}
-			
-			NewKeyPositionScratch.Add(StartPosition);
+
+			NewKeyPositionScratch.Add(StartPosition.Transform(CurveTransform));
 		}
 
 		Curve->SetKeyPositions(KeyData.Handles, NewKeyPositionScratch, EPropertyChangeType::Interactive);
@@ -142,10 +151,16 @@ void FCurveEditorDragOperation_MoveKeys::OnCancelDrag()
 {
 	ICurveEditorKeyDragOperation::OnCancelDrag();
 
-	for (const FKeyData& KeyData : KeysByCurve)
+	for (FKeyData& KeyData : KeysByCurve)
 	{
 		if (FCurveModel* Curve = CurveEditor->FindCurve(KeyData.CurveID))
 		{
+			FTransform2d InverseInitialDragTransform = KeyData.InitialDragTransform.Inverse();
+			for (FKeyPosition& Position : KeyData.StartKeyPositions)
+			{
+				Position = Position.Transform(InverseInitialDragTransform);
+			}
+
 			Curve->SetKeyPositions(KeyData.Handles, KeyData.StartKeyPositions, EPropertyChangeType::ValueSet);
 		}
 	}
diff --git a/Engine/Source/Editor/CurveEditor/Private/DragOperations/CurveEditorDragOperation_MoveKeys.h b/Engine/Source/Editor/CurveEditor/Private/DragOperations/CurveEditorDragOperation_MoveKeys.h
--- a/Engine/Source/Editor/CurveEditor/Private/DragOperations/CurveEditorDragOperation_MoveKeys.h
+++ b/Engine/Source/Editor/CurveEditor/Private/DragOperations/CurveEditorDragOperation_MoveKeys.h
@@ -46,6 +46,8 @@ private:
 		TArray<FKeyPosition> StartKeyPositions;
 		/** Used in OnEndDrag to send final key updates */
 		TArray<FKeyPosition> LastDraggedKeyPositions;
+		/** Initial curve transform for this curve */
+		FTransform2d InitialDragTransform;
 	};
 
 	/** Key dragging data stored per-curve */
diff --git a/Engine/Source/Editor/CurveEditor/Private/DragOperations/CurveEditorDragOperation_Pan.cpp b/Engine/Source/Editor/CurveEditor/Private/DragOperations/CurveEditorDragOperation_Pan.cpp
--- a/Engine/Source/Editor/CurveEditor/Private/DragOperations/CurveEditorDragOperation_Pan.cpp
+++ b/Engine/Source/Editor/CurveEditor/Private/DragOperations/CurveEditorDragOperation_Pan.cpp
@@ -43,7 +43,7 @@ void FCurveEditorDragOperation_PanView::OnDrag(FVector2D InitialPosition, FVecto
 	double OutputMin = InitialOutputMin + PixelDelta.Y / ViewSpace.PixelsPerOutput();
 	double OutputMax = InitialOutputMax + PixelDelta.Y / ViewSpace.PixelsPerOutput();
 
-	CurveEditor->GetBounds().SetInputBounds(InputMin, InputMax);
+	View->SetInputBounds(InputMin, InputMax);
 	View->SetOutputBounds(OutputMin, OutputMax);
 }
 
@@ -94,7 +94,6 @@ void FCurveEditorDragOperation_PanInput::OnDrag(FVector2D InitialPosition, FVect
 
 	double InputMin = InitialInputMin - PixelDelta.X / InputSpace.PixelsPerInput();
 	double InputMax = InitialInputMax - PixelDelta.X / InputSpace.PixelsPerInput();
-
 	CurveEditor->GetBounds().SetInputBounds(InputMin, InputMax);
 
 	CurveEditor->GetPanel()->ScrollBy(-MouseEvent.GetCursorDelta().Y);
