Curve Editor: When toggling selection, respect bPreferPointSelection which prioritizes toggling point selection over tangent selection if there are any points selected. The same thing logic occurs currently for just regular selection (not toggling with control). This fixes an issue where toggling point selection will inadvertently also change tangent selection.

diff --git a/Engine/Source/Editor/CurveEditor/Private/DragOperations/CurveEditorDragOperation_Marquee.cpp b/Engine/Source/Editor/CurveEditor/Private/DragOperations/CurveEditorDragOperation_Marquee.cpp
--- a/Engine/Source/Editor/CurveEditor/Private/DragOperations/CurveEditorDragOperation_Marquee.cpp
+++ b/Engine/Source/Editor/CurveEditor/Private/DragOperations/CurveEditorDragOperation_Marquee.cpp
@@ -154,7 +154,17 @@ void FCurveEditorDragOperation_Marquee::OnEndDrag(FVector2D InitialPosition, FVe
 		}
 		else if (bIsControlDown)
 		{
-			CurveEditor->Selection.Toggle(Point);
+			if (bPreferPointSelection)
+			{
+				if (Point.PointType == ECurvePointType::Key)
+				{
+					CurveEditor->Selection.Toggle(Point);
+				}
+			}
+			else
+			{
+				CurveEditor->Selection.Toggle(Point);
+			}
 		}
 		else
 		{
