Curve Editor: Correct reversed parameters that were causing dragging in the curve editor to feel sketchy

diff --git a/Engine/Source/Editor/CurveEditor/Private/DragOperations/CurveEditorDragOperation_MoveKeys.cpp b/Engine/Source/Editor/CurveEditor/Private/DragOperations/CurveEditorDragOperation_MoveKeys.cpp
--- a/Engine/Source/Editor/CurveEditor/Private/DragOperations/CurveEditorDragOperation_MoveKeys.cpp
+++ b/Engine/Source/Editor/CurveEditor/Private/DragOperations/CurveEditorDragOperation_MoveKeys.cpp
@@ -67,7 +67,7 @@ void FCurveEditorDragOperation_MoveKeys::OnBeginDrag(FVector2D InitialPosition,
 void FCurveEditorDragOperation_MoveKeys::OnDrag(FVector2D InitialPosition, FVector2D CurrentPosition, const FPointerEvent& MouseEvent)
 {
 	TArray<FKeyPosition> NewKeyPositionScratch;
-	FVector2D MousePosition = CurveEditor->GetAxisSnap().GetSnappedPosition(InitialPosition,CurrentPosition, LastMousePosition, MouseEvent, SnappingState);
+	FVector2D MousePosition = CurveEditor->GetAxisSnap().GetSnappedPosition(InitialPosition, LastMousePosition, CurrentPosition, MouseEvent, SnappingState);
 	LastMousePosition = CurrentPosition;
 
 	for (FKeyData& KeyData : KeysByCurve)
