CurveEditor
- Fixed crash while engine is exiting due to FCurveEditor unbinding its uobject delegates in destructor

diff --git a/Engine/Source/Editor/CurveEditor/Private/CurveEditor.cpp b/Engine/Source/Editor/CurveEditor/Private/CurveEditor.cpp
--- a/Engine/Source/Editor/CurveEditor/Private/CurveEditor.cpp
+++ b/Engine/Source/Editor/CurveEditor/Private/CurveEditor.cpp
@@ -91,7 +91,7 @@ FCurveEditor::FCurveEditor()
 
 FCurveEditor::~FCurveEditor()
 {
-	if (Settings)
+	if (!IsEngineExitRequested() && Settings)
 	{
 		Settings->GetOnCustomColorsChanged().RemoveAll(this);
 		Settings->GetOnAxisSnappingChanged().RemoveAll(this);
