Curve Editor: Check for valid key handle before removing

diff --git a/Engine/Source/Editor/CurveEditor/Private/DragOperations/CurveEditorDragOperation_MoveKeys.cpp b/Engine/Source/Editor/CurveEditor/Private/DragOperations/CurveEditorDragOperation_MoveKeys.cpp
--- a/Engine/Source/Editor/CurveEditor/Private/DragOperations/CurveEditorDragOperation_MoveKeys.cpp
+++ b/Engine/Source/Editor/CurveEditor/Private/DragOperations/CurveEditorDragOperation_MoveKeys.cpp
@@ -215,7 +215,7 @@ void FCurveEditorDragOperation_MoveKeys::OnEndDrag(FVector2D InitialPosition, FV
 					{
 						if (KeyInRange != MaxKeyHandle)
 						{
-							KeysToRemove.AddUnique(KeyInRange);
+							KeysToRemove.Add(KeyInRange);
 						}
 					}
 				}
diff --git a/Engine/Source/Editor/CurveEditor/Private/RichCurveEditorModel.cpp b/Engine/Source/Editor/CurveEditor/Private/RichCurveEditorModel.cpp
--- a/Engine/Source/Editor/CurveEditor/Private/RichCurveEditorModel.cpp
+++ b/Engine/Source/Editor/CurveEditor/Private/RichCurveEditorModel.cpp
@@ -208,7 +208,10 @@ void FRichCurveEditorModel::RemoveKeys(TArrayView<const FKeyHandle> InKeys)
 			FRichCurve& RichCurve = GetRichCurve();
 			for (FKeyHandle Handle : InKeys)
 			{
-				RichCurve.DeleteKey(Handle);
+				if (RichCurve.IsKeyHandleValid(Handle))
+				{
+					RichCurve.DeleteKey(Handle);
+				}
 			}
 
 			CurveModifiedDelegate.Broadcast();
