Curve Editor: Remove hack of looking for the first two properties to looking for them by name, Time and Value

This is slightly more robust because we want to expose other properties on the key struct, like interp and tangent modes.

diff --git a/Engine/Source/Editor/CurveEditor/Private/SCurveKeyDetailPanel.cpp b/Engine/Source/Editor/CurveEditor/Private/SCurveKeyDetailPanel.cpp
--- a/Engine/Source/Editor/CurveEditor/Private/SCurveKeyDetailPanel.cpp
+++ b/Engine/Source/Editor/CurveEditor/Private/SCurveKeyDetailPanel.cpp
@@ -95,13 +95,11 @@ private:
 
 void SCurveKeyDetailPanel::PropertyRowsRefreshed()
 {
-	// UE_LOG(LogTemp, Log, TEXT("PropertyRowsRefreshed"));
 	TSharedPtr<SWidget> TimeWidget = nullptr;
 	TSharedPtr<SWidget> ValueWidget = nullptr;
 
 	for (TSharedRef<IDetailTreeNode> RootNode : PropertyRowGenerator->GetRootTreeNodes())
 	{
-		// UE_LOG(LogTemp, Log, TEXT("Root NodeName: %s"), *RootNode->GetNodeName().ToString());
 		TArray<TSharedRef<IDetailTreeNode>> Children;
 		RootNode->GetChildren(Children);
 
@@ -109,16 +107,13 @@ void SCurveKeyDetailPanel::PropertyRowsRefreshed()
 		{
 			TArray<TSharedRef<IDetailTreeNode>> SubChildren;
 			Child->GetChildren(SubChildren);
-			// UE_LOG(LogTemp, Log, TEXT("Child NodeName: %s NumChildren: %d"), *Child->GetNodeName().ToString(), SubChildren.Num());
 
-			// This is an ugly temporary hack until PropertyRowGenerator returns names for customized properties. This uses the first
-			// two fields on the object instead of looking for "Time" and "Value". :(
-			if (!TimeWidget.IsValid())
+			if (!TimeWidget.IsValid() && Child->GetNodeName() == TEXT("Time"))
 			{
 				FNodeWidgets NodeWidgets = Child->CreateNodeWidgets();
 				TimeWidget = NodeWidgets.ValueWidget;
 			}
-			else if (!ValueWidget.IsValid())
+			else if (!ValueWidget.IsValid() && Child->GetNodeName() == TEXT("Value"))
 			{
 				FNodeWidgets NodeWidgets = Child->CreateNodeWidgets();
 				ValueWidget = NodeWidgets.ValueWidget;
