Curve Editor: Key details fix

Follow up fix to 35727872. It turns out the ugly hack is still necessary in some cases

diff --git a/Engine/Source/Editor/CurveEditor/Private/SCurveKeyDetailPanel.cpp b/Engine/Source/Editor/CurveEditor/Private/SCurveKeyDetailPanel.cpp
--- a/Engine/Source/Editor/CurveEditor/Private/SCurveKeyDetailPanel.cpp
+++ b/Engine/Source/Editor/CurveEditor/Private/SCurveKeyDetailPanel.cpp
@@ -121,6 +121,34 @@ void SCurveKeyDetailPanel::PropertyRowsRefreshed()
 		}
 	}
 
+	// If either Time or Value were not found, use this ugly temporary hack until PropertyRowGenerator returns names for customized properties. This uses the first
+	// two fields on the object instead of looking for "Time" and "Value". :(
+	if (!TimeWidget || !ValueWidget)
+	{
+		TimeWidget = nullptr;
+		ValueWidget = nullptr;
+
+		for (TSharedRef<IDetailTreeNode> RootNode : PropertyRowGenerator->GetRootTreeNodes())
+		{
+			TArray<TSharedRef<IDetailTreeNode>> Children;
+			RootNode->GetChildren(Children);
+
+			for (TSharedRef<IDetailTreeNode> Child : Children)
+			{
+				if (!TimeWidget.IsValid())
+				{
+					FNodeWidgets NodeWidgets = Child->CreateNodeWidgets();
+					TimeWidget = NodeWidgets.ValueWidget;
+				}
+				else if (!ValueWidget.IsValid())
+				{
+					FNodeWidgets NodeWidgets = Child->CreateNodeWidgets();
+					ValueWidget = NodeWidgets.ValueWidget;
+				}
+			}
+		}
+	}
+
 	if (!TimeWidget)
 	{
 		if (!TempTimeWidget.IsValid())
