Curve Editor: Check for any grid lines before Algo::MinElement when snapping to the grid

This fixes a bug in inserting a key into space switching curves in the curve editor since there aren't any grid lines to snap to.

diff --git a/Engine/Source/Editor/CurveEditor/Public/CurveEditorSnapMetrics.h b/Engine/Source/Editor/CurveEditor/Public/CurveEditorSnapMetrics.h
--- a/Engine/Source/Editor/CurveEditor/Public/CurveEditorSnapMetrics.h
+++ b/Engine/Source/Editor/CurveEditor/Public/CurveEditorSnapMetrics.h
@@ -38,7 +38,7 @@ struct FCurveSnapMetrics
 	/** Snap the specified output value to the output snap interval if necessary */
 	FORCEINLINE double SnapOutput(double OutputValue)
 	{
-		return bSnapOutputValues ? *Algo::MinElement(AllGridLines, 
+		return bSnapOutputValues && !AllGridLines.IsEmpty() ? *Algo::MinElement(AllGridLines,
 			[OutputValue](double Val1, double Val2) { return FMath::Abs(Val1 - OutputValue) < FMath::Abs(Val2 - OutputValue); }
 		) : OutputValue;
 	}
