CurveEditor: Add hotkeys for matching last keys tangents with first keys tangents and vice versa

diff --git a/Engine/Source/Editor/CurveEditor/Private/CurveEditor.cpp b/Engine/Source/Editor/CurveEditor/Private/CurveEditor.cpp
--- a/Engine/Source/Editor/CurveEditor/Private/CurveEditor.cpp
+++ b/Engine/Source/Editor/CurveEditor/Private/CurveEditor.cpp
@@ -463,6 +463,9 @@ void FCurveEditor::BindCommands()
 	CommandList->MapAction(FCurveEditorCommands::Get().SelectNone, FExecuteAction::CreateSP(this, &FCurveEditor::SelectNone));
 	CommandList->MapAction(FCurveEditorCommands::Get().InvertSelection, FExecuteAction::CreateSP(this, &FCurveEditor::InvertSelection));
 
+	CommandList->MapAction(FCurveEditorCommands::Get().MatchLastTangentToFirst, FExecuteAction::CreateSP(this, &FCurveEditor::MatchLastTangentToFirst,true));
+	CommandList->MapAction(FCurveEditorCommands::Get().MatchFirstTangentToLast, FExecuteAction::CreateSP(this, &FCurveEditor::MatchLastTangentToFirst,false));
+
 	{
 		FExecuteAction   ToggleInputSnapping     = FExecuteAction::CreateSP(this,   &FCurveEditor::ToggleInputSnapping);
 		FIsActionChecked IsInputSnappingEnabled  = FIsActionChecked::CreateSP(this, &FCurveEditor::IsInputSnappingEnabled);
@@ -1509,6 +1512,51 @@ public:
 	TArray<UCurveEditorCopyBuffer*> NewCopyBuffers;
 };
 
+void FCurveEditor::MatchLastTangentToFirst(bool bMatchLastToFirst)
+{
+	FScopedTransaction Transaction(LOCTEXT("MatchTangents", "Match Tangents"));
+	bool bFoundAnyTangents = false;
+
+	TArray<FKeyHandle> KeyHandles;
+	KeyHandles.SetNum(2);
+	TArray<FKeyAttributes> KeyAttributes;
+	KeyAttributes.SetNum(2);
+	TArray<FKeyHandle> AllKeyHandles;
+
+	for (const TTuple<FCurveModelID, TUniquePtr<FCurveModel>>& Pair : CurveData)
+	{
+		FCurveModel* CurveModel = Pair.Value.Get();
+
+		if (CurveModel)
+		{
+			// Get all of the key handles from this curve.
+			CurveModel->GetKeys(TNumericLimits<double>::Lowest(), TNumericLimits<double>::Max(), TNumericLimits<double>::Lowest(), TNumericLimits<double>::Max(), AllKeyHandles);
+			//need 2
+			if (AllKeyHandles.Num() < 2)
+			{
+				continue;
+			}
+			bFoundAnyTangents = true;
+			KeyHandles[0] = AllKeyHandles[0];
+			KeyHandles[1] = AllKeyHandles[AllKeyHandles.Num() - 1];
+			
+			CurveModel->GetKeyAttributes(KeyHandles, KeyAttributes);
+			if (bMatchLastToFirst)
+			{
+				KeyAttributes[1] = KeyAttributes[0];
+			}
+			else
+			{
+				KeyAttributes[0] = KeyAttributes[1];
+			}
+			CurveModel->SetKeyAttributes(KeyHandles, KeyAttributes);
+		}
+	}
+	if (!bFoundAnyTangents)
+	{
+		Transaction.Cancel();
+	}
+}
 
 bool FCurveEditor::CanPaste(const FString& TextToImport) const
 {
diff --git a/Engine/Source/Editor/CurveEditor/Private/CurveEditorCommands.cpp b/Engine/Source/Editor/CurveEditor/Private/CurveEditorCommands.cpp
--- a/Engine/Source/Editor/CurveEditor/Private/CurveEditorCommands.cpp
+++ b/Engine/Source/Editor/CurveEditor/Private/CurveEditorCommands.cpp
@@ -33,6 +33,9 @@ void FCurveEditorCommands::RegisterCommands()
 	UI_COMMAND(FlattenTangents, "Flatten", "Flatten Tangents", EUserInterfaceActionType::Button, FInputChord(EKeys::Six));
 	UI_COMMAND(StraightenTangents, "Straighten", "Straighten Tangents", EUserInterfaceActionType::Button, FInputChord());
 
+	UI_COMMAND(MatchLastTangentToFirst, "Match Last Tangent To First", "Match last keys tangents to the first keys tangents", EUserInterfaceActionType::Button, FInputChord());
+	UI_COMMAND(MatchFirstTangentToLast, "Match First Tangent To Last", "Match first keys tangents to the last keys tangents", EUserInterfaceActionType::Button, FInputChord());
+
 	UI_COMMAND(BakeCurve, "Bake", "Bake curve", EUserInterfaceActionType::Button, FInputChord());
 	UI_COMMAND(ReduceCurve, "Reduce", "Reduce curve", EUserInterfaceActionType::Button, FInputChord());
 
diff --git a/Engine/Source/Editor/CurveEditor/Public/CurveEditor.h b/Engine/Source/Editor/CurveEditor/Public/CurveEditor.h
--- a/Engine/Source/Editor/CurveEditor/Public/CurveEditor.h
+++ b/Engine/Source/Editor/CurveEditor/Public/CurveEditor.h
@@ -539,6 +539,11 @@ public:
 	 */
 	void StraightenSelection();
 
+	/** 
+	* Sets the last keys tangents to match the first, if bFirstToLast is false do the opposite
+	*/
+	void MatchLastTangentToFirst(bool bLastToFirst = true);
+
 	/**
 	 * Set random curve colors
 	 */
diff --git a/Engine/Source/Editor/CurveEditor/Public/CurveEditorCommands.h b/Engine/Source/Editor/CurveEditor/Public/CurveEditorCommands.h
--- a/Engine/Source/Editor/CurveEditor/Public/CurveEditorCommands.h
+++ b/Engine/Source/Editor/CurveEditor/Public/CurveEditorCommands.h
@@ -51,6 +51,9 @@ public:
 	TSharedPtr<FUICommandInfo> FlattenTangents;
 	TSharedPtr<FUICommandInfo> StraightenTangents;
 
+	TSharedPtr< FUICommandInfo > MatchLastTangentToFirst;
+	TSharedPtr< FUICommandInfo > MatchFirstTangentToLast;
+	
 	TSharedPtr<FUICommandInfo> BakeCurve;
 	TSharedPtr<FUICommandInfo> ReduceCurve;
 
