Sequencer: Curve Editor: Disable throttling when using scrub time hotkey

diff --git a/Engine/Source/Editor/CurveEditor/Private/DragOperations/CurveEditorDragOperation_ScrubTime.cpp b/Engine/Source/Editor/CurveEditor/Private/DragOperations/CurveEditorDragOperation_ScrubTime.cpp
--- a/Engine/Source/Editor/CurveEditor/Private/DragOperations/CurveEditorDragOperation_ScrubTime.cpp
+++ b/Engine/Source/Editor/CurveEditor/Private/DragOperations/CurveEditorDragOperation_ScrubTime.cpp
@@ -9,6 +9,8 @@
 #include "Math/UnrealMathUtility.h"
 #include "SCurveEditorPanel.h"
 #include "SCurveEditorView.h"
+#include "Application/ThrottleManager.h"
+
 
 FCurveEditorDragOperation_ScrubTime::FCurveEditorDragOperation_ScrubTime(FCurveEditor* InCurveEditor)
 	: CurveEditor(InCurveEditor) 
@@ -16,6 +18,7 @@ FCurveEditorDragOperation_ScrubTime::FCurveEditorDragOperation_ScrubTime(FCurveE
 
 void FCurveEditorDragOperation_ScrubTime::OnBeginDrag(FVector2D InitialPosition, FVector2D CurrentPosition, const FPointerEvent& MouseEvent)
 {
+	FSlateThrottleManager::Get().DisableThrottle(true);
 	TSharedPtr<ITimeSliderController> TimeSliderController = CurveEditor->GetTimeSliderController();
 	if (TimeSliderController)
 	{
@@ -60,6 +63,7 @@ void FCurveEditorDragOperation_ScrubTime::OnEndDrag(FVector2D InitialPosition, F
 	{
 		TimeSliderController->SetPlaybackStatus(ETimeSliderPlaybackStatus::Stopped);
 	}
+	FSlateThrottleManager::Get().DisableThrottle(false);
 }
 
 void FCurveEditorDragOperation_ScrubTime::OnCancelDrag()
