Curve Editor: Specify a minimum size for the splitter so that the user doesn't get into an unrecoverable state

#rnx

diff --git a/Engine/Source/Editor/CurveEditor/Private/SCurveEditorPanel.cpp b/Engine/Source/Editor/CurveEditor/Private/SCurveEditorPanel.cpp
--- a/Engine/Source/Editor/CurveEditor/Private/SCurveEditorPanel.cpp
+++ b/Engine/Source/Editor/CurveEditor/Private/SCurveEditorPanel.cpp
@@ -286,12 +286,14 @@ void SCurveEditorPanel::Construct(const FArguments& InArgs, TSharedRef<FCurveEdi
 
 			+ SSplitter::Slot()
 			.Value(FillCoefficient_0)
+			.MinSize(1.f)
 			.OnSlotResized(SSplitter::FOnSlotResized::CreateSP(this, &SCurveEditorPanel::OnColumnFillCoefficientChanged, 0))
 			[
 				InArgs._TreeContent.Widget
 			]
 
 			+ SSplitter::Slot()
+			.MinSize(1.f)
 			.Value(FillCoefficient_1)
 			.OnSlotResized(SSplitter::FOnSlotResized::CreateSP(this, &SCurveEditorPanel::OnColumnFillCoefficientChanged, 1))
 			[
