Curve Editor: Simplify FCurveEditorDragOperation_Marquee::OnEndDrag by extracting functions.

diff --git a/Engine/Source/Editor/CurveEditor/Private/DragOperations/CurveEditorDragOperation_Marquee.cpp b/Engine/Source/Editor/CurveEditor/Private/DragOperations/CurveEditorDragOperation_Marquee.cpp
--- a/Engine/Source/Editor/CurveEditor/Private/DragOperations/CurveEditorDragOperation_Marquee.cpp
+++ b/Engine/Source/Editor/CurveEditor/Private/DragOperations/CurveEditorDragOperation_Marquee.cpp
@@ -2,6 +2,7 @@
 
 #include "DragOperations/CurveEditorDragOperation_Marquee.h"
 
+#include "Algo/AnyOf.h"
 #include "Containers/Array.h"
 #include "Containers/ArrayView.h"
 #include "Containers/Map.h"
@@ -21,6 +22,84 @@
 #include "Templates/SharedPointer.h"
 #include "Templates/Tuple.h"
 
+namespace UE::CurveEditor::MarqueeDetail
+{
+static TArray<FCurvePointHandle> FindOverlappingPoints(
+	const FSlateRect& Marquee, const FCurveEditor* CurveEditor, const SCurveEditorView* LockedToView
+	)
+{
+	TArray<FCurvePointHandle> AllPoints;
+	
+	if (LockedToView)
+	{
+		LockedToView->GetPointsThenCurveWithinWidgetRange(Marquee, &AllPoints);
+	}
+	else
+	{
+		const TSharedPtr<SCurveEditorPanel> CurveEditorPanel = CurveEditor->GetPanel();
+
+		const FGeometry ViewContainerGeometry = CurveEditorPanel->GetViewContainerGeometry();
+		const FSlateLayoutTransform InverseContainerTransform = ViewContainerGeometry.GetAccumulatedLayoutTransform().Inverse();
+		for (TSharedPtr<SCurveEditorView> View : CurveEditorPanel->GetViews())
+		{
+			const FGeometry& LocalGeometry = View->GetCachedGeometry();
+			const FSlateLayoutTransform ContainerToView = InverseContainerTransform.Concatenate(
+				LocalGeometry.GetAccumulatedLayoutTransform()
+				).Inverse();
+
+			const FSlateRect UnclippedLocalMarquee = FSlateRect(
+				ContainerToView.TransformPoint(Marquee.GetTopLeft2f()),
+				ContainerToView.TransformPoint(Marquee.GetBottomRight2f())
+				);
+			const FSlateRect ClippedLocalMarquee = UnclippedLocalMarquee.IntersectionWith(
+				FSlateRect(FVector2D(0.f,0.f), LocalGeometry.GetLocalSize())
+				);
+
+			if (ClippedLocalMarquee.IsValid() && !ClippedLocalMarquee.IsEmpty())
+			{
+				View->GetPointsThenCurveWithinWidgetRange(ClippedLocalMarquee, &AllPoints);
+			}
+		}
+	}
+
+	return AllPoints;
+}
+
+static void RestrictSelectionToEitherPointsOrTangents(bool bPreferPointSelection, FCurveEditor* CurveEditor)
+{
+	TArray<FCurvePointHandle> CurvePointsToRemove;
+	for (const TTuple<FCurveModelID, FKeyHandleSet>& Pair : CurveEditor->GetSelection().GetAll())
+	{
+		for (FKeyHandle Handle : Pair.Value.AsArray())
+		{
+			ECurvePointType PointType = Pair.Value.PointType(Handle);
+
+			if (bPreferPointSelection)
+			{
+				// If selecting points, deselect tangen handles (ie. anything that's not a point/key)
+				if (PointType != ECurvePointType::Key)
+				{
+					CurvePointsToRemove.Add(FCurvePointHandle(Pair.Key, PointType, Handle));
+				}
+			}
+			else
+			{
+				// Otherwise when selecting tangent handles, deselect anything that's a key
+				if (PointType == ECurvePointType::Key)
+				{
+					CurvePointsToRemove.Add(FCurvePointHandle(Pair.Key, PointType, Handle));
+				}
+			}
+		}
+	}
+
+	for (const FCurvePointHandle& Point : CurvePointsToRemove)
+	{
+		CurveEditor->Selection.Remove(Point);
+	}
+}
+}
+
 FCurveEditorDragOperation_Marquee::FCurveEditorDragOperation_Marquee(FCurveEditor* InCurveEditor)
 	: LockedToView(nullptr)
 {
@@ -57,38 +136,7 @@ void FCurveEditorDragOperation_Marquee::OnDrag(FVector2D InitialPosition, FVecto
 
 void FCurveEditorDragOperation_Marquee::OnEndDrag(FVector2D InitialPosition, FVector2D CurrentPosition, const FPointerEvent& MouseEvent)
 {
-	TArray<FCurvePointHandle> AllPoints;
-
-	if (LockedToView)
-	{
-		if (!LockedToView->GetPointsWithinWidgetRange(Marquee, &AllPoints))
-		{
-			LockedToView->GetCurveWithinWidgetRange(Marquee, &AllPoints);
-		}
-	}
-	else
-	{
-		TSharedPtr<SCurveEditorPanel> CurveEditorPanel = CurveEditor->GetPanel();
-
-		FGeometry ViewContainerGeometry = CurveEditorPanel->GetViewContainerGeometry();
-		FSlateLayoutTransform InverseContainerTransform = ViewContainerGeometry.GetAccumulatedLayoutTransform().Inverse();
-		for (TSharedPtr<SCurveEditorView> View : CurveEditorPanel->GetViews())
-		{
-			const FGeometry& LocalGeometry = View->GetCachedGeometry();
-			FSlateLayoutTransform ContainerToView = InverseContainerTransform.Concatenate(LocalGeometry.GetAccumulatedLayoutTransform()).Inverse();
-
-			FSlateRect UnclippedLocalMarquee = FSlateRect(ContainerToView.TransformPoint(Marquee.GetTopLeft2f()), ContainerToView.TransformPoint(Marquee.GetBottomRight2f()));
-			FSlateRect ClippedLocalMarquee = UnclippedLocalMarquee.IntersectionWith(FSlateRect(FVector2D(0.f,0.f), LocalGeometry.GetLocalSize()));
-
-			if (ClippedLocalMarquee.IsValid() && !ClippedLocalMarquee.IsEmpty())
-			{
-				if (!View->GetPointsWithinWidgetRange(ClippedLocalMarquee, &AllPoints))
-				{
-					View->GetCurveWithinWidgetRange(ClippedLocalMarquee, &AllPoints);
-				}
-			}
-		}
-	}
+	const TArray<FCurvePointHandle> AllPoints = UE::CurveEditor::MarqueeDetail::FindOverlappingPoints(Marquee, CurveEditor, LockedToView);
 
 	const bool bIsShiftDown = MouseEvent.IsShiftDown();
 	const bool bIsAltDown = MouseEvent.IsAltDown();
@@ -100,49 +148,15 @@ void FCurveEditorDragOperation_Marquee::OnEndDrag(FVector2D InitialPosition, FVe
 	}
 
 	// If there are any points to be selected, prefer selecting points over tangents
-	bool bPreferPointSelection = false;
-	for (const FCurvePointHandle& Point : AllPoints)
+	const bool bPreferPointSelection = Algo::AnyOf(AllPoints, [](const FCurvePointHandle& Point)
 	{
-		if (Point.PointType == ECurvePointType::Key)
-		{
-			bPreferPointSelection = true;
-			break;
-		}
-	}
+		return Point.PointType == ECurvePointType::Key;
+	});
 
 	// When adding to the existing selection, ensure only either points or tangents are selected
-	TArray<FCurvePointHandle> CurvePointsToRemove;
 	if (bIsShiftDown)
 	{
-		for (const TTuple<FCurveModelID, FKeyHandleSet>& Pair : CurveEditor->GetSelection().GetAll())
-		{
-			for (FKeyHandle Handle : Pair.Value.AsArray())
-			{
-				ECurvePointType PointType = Pair.Value.PointType(Handle);
-
-				if (bPreferPointSelection)
-				{
-					// If selecting points, deselect tangen handles (ie. anything that's not a point/key)
-					if (PointType != ECurvePointType::Key)
-					{
-						CurvePointsToRemove.Add(FCurvePointHandle(Pair.Key, PointType, Handle));
-					}
-				}
-				else
-				{
-					// Otherwise when selecting tangent handles, deselect anything that's a key
-					if (PointType == ECurvePointType::Key)
-					{
-						CurvePointsToRemove.Add(FCurvePointHandle(Pair.Key, PointType, Handle));
-					}
-				}
-			}
-		}
-	}
-
-	for (const FCurvePointHandle& Point : CurvePointsToRemove)
-	{
-		CurveEditor->Selection.Remove(Point);
+		UE::CurveEditor::MarqueeDetail::RestrictSelectionToEitherPointsOrTangents(bPreferPointSelection, CurveEditor);
 	}
 
 	// Now that we've gathered the overlapping points, perform the relevant selection
@@ -166,21 +180,19 @@ void FCurveEditorDragOperation_Marquee::OnEndDrag(FVector2D InitialPosition, FVe
 				CurveEditor->Selection.Toggle(Point);
 			}
 		}
-		else
+		else if (bPreferPointSelection)
 		{
-			if (bPreferPointSelection)
-			{
-				if (Point.PointType == ECurvePointType::Key)
-				{
-					CurveEditor->Selection.Add(Point);
-				}
-			}
-			else
+			if (Point.PointType == ECurvePointType::Key)
 			{
 				CurveEditor->Selection.Add(Point);
 			}
 		}
+		else
+		{
+			CurveEditor->Selection.Add(Point);
+		}
 	}
+	
 }
 
 void FCurveEditorDragOperation_Marquee::OnPaint(const FGeometry& AllottedGeometry, FSlateWindowElementList& OutDrawElements, int32 PaintOnLayerId)
diff --git a/Engine/Source/Editor/CurveEditor/Public/SCurveEditorView.h b/Engine/Source/Editor/CurveEditor/Public/SCurveEditorView.h
--- a/Engine/Source/Editor/CurveEditor/Public/SCurveEditorView.h
+++ b/Engine/Source/Editor/CurveEditor/Public/SCurveEditorView.h
@@ -276,6 +276,20 @@ public:
 	virtual bool GetCurveWithinWidgetRange(const FSlateRect& WidgetRectangle, TArray<FCurvePointHandle>* OutPoints) const
 	{ return false; }
 
+	/**
+	 * Tries to retrieve all curve points that overlap the rectangle.
+	 * If no points are found, selects the overlapping curves instead.
+	 *
+	 * @param WidgetRectangle  The rectangle to hit test against. May not hit points that would exist outside of the view's visible bounds.
+	 * @param OutPoints        (required) pointer to an array to populate with overlapping points
+	 * 
+	 * @return Whether anything was selected.
+	 */
+	bool GetPointsThenCurveWithinWidgetRange(const FSlateRect& WidgetRectangle, TArray<FCurvePointHandle>* OutPoints) const
+	{
+		return GetPointsWithinWidgetRange(WidgetRectangle, OutPoints) || GetCurveWithinWidgetRange(WidgetRectangle, OutPoints);
+	}
+
 	/**
 	 * Retrieve the id of the hovered curve
 	 */
