Curve Editor: Adjust sensitivity of when marquee selects points.
- Adds MarqueePointSensitivity variable to the Curve Editor settings

diff --git a/Engine/Source/Editor/CurveEditor/Private/CurveEditorSettings.cpp b/Engine/Source/Editor/CurveEditor/Private/CurveEditorSettings.cpp
--- a/Engine/Source/Editor/CurveEditor/Private/CurveEditorSettings.cpp
+++ b/Engine/Source/Editor/CurveEditor/Private/CurveEditorSettings.cpp
@@ -336,6 +336,15 @@ void UCurveEditorSettings::SetTreeViewWidth(float InTreeViewWidth)
 	}
 }
 
+void UCurveEditorSettings::SetMarqueePointSensitivity(float InMarqueePointSensitivity)
+{
+	if (InMarqueePointSensitivity != MarqueePointSensitivity)
+	{
+		MarqueePointSensitivity = InMarqueePointSensitivity;
+		SaveConfig();
+	}
+}
+
 void UCurveEditorSettings::PostEditChangeProperty(FPropertyChangedEvent& PropertyChangedEvent)
 {
 	FName PropertyName = (PropertyChangedEvent.Property != NULL) ? PropertyChangedEvent.Property->GetFName() : NAME_None;
diff --git a/Engine/Source/Editor/CurveEditor/Private/Views/SInteractiveCurveEditorView.cpp b/Engine/Source/Editor/CurveEditor/Private/Views/SInteractiveCurveEditorView.cpp
--- a/Engine/Source/Editor/CurveEditor/Private/Views/SInteractiveCurveEditorView.cpp
+++ b/Engine/Source/Editor/CurveEditor/Private/Views/SInteractiveCurveEditorView.cpp
@@ -943,6 +943,8 @@ bool SInteractiveCurveEditorView::GetPointsWithinWidgetRange(const FSlateRect& W
 	FVector Start, End, StartToEnd;
 	FBox WidgetRectangleBox(FVector(WidgetRectangle.Left, WidgetRectangle.Top, 0), FVector(WidgetRectangle.Right, WidgetRectangle.Bottom, 0));
 
+	const float PointOverlapSensitivity = CurveEditor->GetSettings()->GetMarqueePointSensitivity();
+	
 	// Iterate through all of our points and see which points the marquee overlaps. Both of these coordinate systems
 	// are in screen space pixels.  Also check tangent lines
 	bool bFound = false;
@@ -953,7 +955,8 @@ bool SInteractiveCurveEditorView::GetPointsWithinWidgetRange(const FSlateRect& W
 			const FCurvePointInfo& Point = DrawParams.Points[PointIndex];
 
 			const FKeyDrawInfo& DrawInfo = DrawParams.GetKeyDrawInfo(Point.Type, PointIndex);
-			FSlateRect PointRect = FSlateRect::FromPointAndExtent(Point.ScreenPosition - DrawInfo.ScreenSize/2, DrawInfo.ScreenSize);
+			const FVector2d CollisionScreenSize = DrawInfo.ScreenSize * PointOverlapSensitivity;
+			const FSlateRect PointRect = FSlateRect::FromPointAndExtent(Point.ScreenPosition - CollisionScreenSize/2, CollisionScreenSize);
 
 			if (FSlateRect::DoRectanglesIntersect(PointRect, WidgetRectangle))
 			{
diff --git a/Engine/Source/Editor/CurveEditor/Public/CurveEditorSettings.h b/Engine/Source/Editor/CurveEditor/Public/CurveEditorSettings.h
--- a/Engine/Source/Editor/CurveEditor/Public/CurveEditorSettings.h
+++ b/Engine/Source/Editor/CurveEditor/Public/CurveEditorSettings.h
@@ -180,6 +180,11 @@ public:
 	/** Sets the tree view width percentage */
 	void SetTreeViewWidth(float InTreeViewWidth);
 
+	/** Gets how sensitive the selection marquee should be when selecting points. */
+	float GetMarqueePointSensitivity() const { return FMath::Clamp(MarqueePointSensitivity, 0.f, 1.f); }
+	/** Sets how sensitive the selection marquee should be when selecting points. */
+	void SetMarqueePointSensitivity(float InMarqueePointSensitivity);
+
 protected:
 
 	virtual void PostEditChangeProperty(FPropertyChangedEvent& PropertyChangedEvent) override;
@@ -234,6 +239,16 @@ protected:
 	UPROPERTY(config, EditAnywhere, Category = "Curve Editor")
 	float TreeViewWidth;
 
+	/**
+	 * Determines how close you must move the selection marque to the center of a point in order to select it.
+	 * 
+	 * This is the percentage of point's center to the point's widget border that must be overlapped with the marquee in order for the point to be selected.
+	 * 1.0 means as soon as marquee overlaps any portion of the point.
+	 * 0.0 means you must touch the point's center to select it.
+	 */
+	UPROPERTY(config, EditAnywhere, Category = "Curve Editor", meta = (ClampMin = 0.0, ClampMax = 1.0))
+	float MarqueePointSensitivity = 0.45f;
+
 private:
 	FOnCustomColorsChanged OnCustomColorsChangedEvent;
 	FOnAxisSnappingChanged OnAxisSnappingChangedEvent;
