Curve Editor: Swapped the behaviour of Ctrl+V and Shift+Ctrl+V when pasting keys.
- Ctrl+V now replaces (new) keys in the pasted range instead of merging them (old).
- Shift+Ctrl+V now merges (new) the pasted keys.

diff --git a/Engine/Source/Editor/CurveEditor/Private/CurveEditor.cpp b/Engine/Source/Editor/CurveEditor/Private/CurveEditor.cpp
--- a/Engine/Source/Editor/CurveEditor/Private/CurveEditor.cpp
+++ b/Engine/Source/Editor/CurveEditor/Private/CurveEditor.cpp
@@ -432,8 +432,8 @@ void FCurveEditor::BindCommands()
 
 	CommandList->MapAction(FGenericCommands::Get().Cut, FExecuteAction::CreateSP(this, &FCurveEditor::CutSelection));
 	CommandList->MapAction(FGenericCommands::Get().Copy, FExecuteAction::CreateSP(this, &FCurveEditor::CopySelection));
-	CommandList->MapAction(FGenericCommands::Get().Paste, FExecuteAction::CreateSP(this, &FCurveEditor::PasteKeys, TSet<FCurveModelID>(), false));
-	CommandList->MapAction(FCurveEditorCommands::Get().PasteOverwriteRange, FExecuteAction::CreateSP(this, &FCurveEditor::PasteKeys, TSet<FCurveModelID>(), true));
+	CommandList->MapAction(FGenericCommands::Get().Paste, FExecuteAction::CreateSP(this, &FCurveEditor::PasteKeys, TSet<FCurveModelID>(), true));
+	CommandList->MapAction(FCurveEditorCommands::Get().PasteOverwriteRange, FExecuteAction::CreateSP(this, &FCurveEditor::PasteKeys, TSet<FCurveModelID>(), false));
 
 	CommandList->MapAction(FCurveEditorCommands::Get().ZoomToFit, FExecuteAction::CreateSP(this, &FCurveEditor::ZoomToFit, EAxisList::All));
 	CommandList->MapAction(FCurveEditorCommands::Get().ZoomToFitHorizontal, FExecuteAction::CreateSP(this, &FCurveEditor::ZoomToFit, EAxisList::X));
