Curve Editor Bake Filter:  Make sure to always bake on whole frames, fix precision issues using doubles, use default interpolation type instead of linear.

diff --git a/Engine/Source/Editor/CurveEditor/Private/Filters/CurveEditorBakeFilter.cpp b/Engine/Source/Editor/CurveEditor/Private/Filters/CurveEditorBakeFilter.cpp
--- a/Engine/Source/Editor/CurveEditor/Private/Filters/CurveEditorBakeFilter.cpp
+++ b/Engine/Source/Editor/CurveEditor/Private/Filters/CurveEditorBakeFilter.cpp
@@ -102,8 +102,21 @@ void UCurveEditorBakeFilter::ApplyFilter_Impl(TSharedRef<FCurveEditor> InCurveEd
 		// Find the hull of the range of the selected keys
 		double MinKey = TNumericLimits<double>::Max(), MaxKey = TNumericLimits<double>::Lowest();
 
+		// Hull of the sampling, maybe different than the hull of the keys
+		double SampleMinKey = MinKey, SampleMaxKey = MaxKey;
+
+		//we want the first key to be on whole frame if it's not we delete it
+		bool bDeleteFirstKey = false;
+		//similarly if the last frame is not on a whole key it may be out of range and also needs to be deleted.
+		bool bDeleteLastKey = false;
+		FFrameTime MinInDisplayFrames(0, 0.0);
+		FFrameTime MaxInDisplayFrames(0, 0.0);
+
 		if (bCustomRangeOverride)
 		{
+			//with custom range we just wipe everything out
+			bDeleteFirstKey = true;
+			bDeleteLastKey = true;
 			if (bUseSeconds)
 			{
 				MinKey = CustomRangeMinInSeconds;
@@ -111,8 +124,10 @@ void UCurveEditorBakeFilter::ApplyFilter_Impl(TSharedRef<FCurveEditor> InCurveEd
 			}
 			else
 			{
-				MinKey = DisplayRate.AsSeconds(FFrameRate::TransformTime(CustomRange.Min, TickResolution, DisplayRate));
-				MaxKey = DisplayRate.AsSeconds(FFrameRate::TransformTime(CustomRange.Max, TickResolution, DisplayRate));
+				MinInDisplayFrames = FFrameRate::TransformTime(CustomRange.Min, TickResolution, DisplayRate);
+				MaxInDisplayFrames = FFrameRate::TransformTime(CustomRange.Max, TickResolution, DisplayRate);
+				MinKey = DisplayRate.AsSeconds(MinInDisplayFrames);
+				MaxKey = DisplayRate.AsSeconds(MaxInDisplayFrames);
 			}
 
 			if (MaxKey < MinKey)
@@ -121,13 +136,29 @@ void UCurveEditorBakeFilter::ApplyFilter_Impl(TSharedRef<FCurveEditor> InCurveEd
 				MinKey = MaxKey;
 				MaxKey = TempMinKey;
 			}
+			SampleMinKey = MinKey;
+			SampleMaxKey = MaxKey;
 		}
 		else
 		{
 			for (FKeyPosition Key : SelectedKeyPositions)
 			{
-				MinKey = FMath::Min(Key.InputValue, MinKey);
-				MaxKey = FMath::Max(Key.InputValue, MaxKey);
+				SampleMinKey = MinKey = FMath::Min(Key.InputValue, MinKey);
+				SampleMaxKey = MaxKey = FMath::Max(Key.InputValue, MaxKey);
+			}
+			MinInDisplayFrames = DisplayRate.AsFrameTime(SampleMinKey);
+			if(MinInDisplayFrames.GetSubFrame() > 0.0)
+			{
+				bDeleteFirstKey = true;
+				MinInDisplayFrames = FFrameTime(MinInDisplayFrames.RoundToFrame(), 0.0);
+				SampleMinKey = DisplayRate.AsSeconds(MinInDisplayFrames);
+			}
+			MaxInDisplayFrames = DisplayRate.AsFrameTime(SampleMaxKey);
+			if (MaxInDisplayFrames.GetSubFrame() > 0.0)
+			{
+				bDeleteLastKey = true;
+				MaxInDisplayFrames = FFrameTime(MaxInDisplayFrames.RoundToFrame(), 0.0);
+				SampleMaxKey = DisplayRate.AsSeconds(MaxInDisplayFrames);
 			}
 		}
 
@@ -139,39 +170,34 @@ void UCurveEditorBakeFilter::ApplyFilter_Impl(TSharedRef<FCurveEditor> InCurveEd
 		
 		double Interval = bUseSeconds ? BakeIntervalInSeconds : DisplayRate.AsSeconds(FFrameRate::TransformTime(BakeInterval, TickResolution, DisplayRate));
 
-		int32 NumKeysToAdd = FMath::FloorToInt((MaxKey - MinKey) / Interval) + 1;
+		int32 KeyIndex = 0;
+		if (bDeleteFirstKey == false)
+		{
+			++KeyIndex;
+		}
+
+		int32 NumKeysToAdd = bUseSeconds ? FMath::RoundToInt((SampleMaxKey - SampleMinKey) / Interval) + 1 :
+			(MaxInDisplayFrames.FrameNumber.Value - MinInDisplayFrames.FrameNumber.Value)  + 1;
 
+	
 		NewKeyPositions.Reset(NumKeysToAdd);
 		NewKeyAttributes.Reset(NumKeysToAdd);
 
-		int32 KeyIndex = 1;
-		for (; KeyIndex < NumKeysToAdd-1; ++KeyIndex)
+		// Get the default interpolation type at the first frame, use this for all of the new keys.
+		TPair<ERichCurveInterpMode, ERichCurveTangentMode> Modes = Curve->GetInterpolationMode(MinKey, ERichCurveInterpMode::RCIM_Cubic, ERichCurveTangentMode::RCTM_SmartAuto);
+		FKeyAttributes DefaultKeyAttribute;
+		DefaultKeyAttribute.SetInterpMode(Modes.Key);
+		DefaultKeyAttribute.SetTangentMode(Modes.Value);
+		
+		for (; KeyIndex < NumKeysToAdd; ++KeyIndex)
 		{
-			double NewTime = MinKey + KeyIndex * Interval;
+			double NewTime = SampleMinKey + KeyIndex * Interval;
 
 			FKeyPosition CurrentKey(NewTime, 0.0);
 			if (Curve->Evaluate(CurrentKey.InputValue, CurrentKey.OutputValue))
 			{
 				NewKeyPositions.Add(CurrentKey);
-				NewKeyAttributes.Add(FKeyAttributes().SetInterpMode(RCIM_Linear));
-			}
-		}
-
-		// Handle the last key
-		if (KeyIndex < NumKeysToAdd)
-		{
-			double NewTime = MinKey + KeyIndex * Interval;
-
-			// Don't create the last key if it has the same value as the last key. This avoids duplicates of the last
-			// key, but improves behavior when the interval doesn't align.
-			if (SelectedKeyPositions.Num() == 0 || !FMath::IsNearlyEqual(NewTime, SelectedKeyPositions.Last().InputValue))
-			{
-				FKeyPosition CurrentKey(NewTime, 0.0);
-				if (Curve->Evaluate(CurrentKey.InputValue, CurrentKey.OutputValue))
-				{
-					NewKeyPositions.Add(CurrentKey);
-					NewKeyAttributes.Add(FKeyAttributes().SetInterpMode(RCIM_Linear));
-				}
+				NewKeyAttributes.Add(DefaultKeyAttribute);
 			}
 		}
 
@@ -185,31 +211,33 @@ void UCurveEditorBakeFilter::ApplyFilter_Impl(TSharedRef<FCurveEditor> InCurveEd
 		TArrayView<TOptional<FKeyHandle>> NewKeyHandlesView = MakeArrayView(NewKeyHandles);
 		FKeyHandleSet& OutHandleSet = OutKeysToSelect.Add(Pair.Key);
 
-		// Manually add the first and last keys to the out set before we remove them below.
-		if (KeyHandles.Num() > 0)
+		// If not deleting, manually add the first and last keys to the out set before we remove them below.
+		if (bDeleteFirstKey == false && KeyHandles.Num() > 0)
 		{
 			OutHandleSet.Add(KeyHandles[0], ECurvePointType::Key);
 		}
-		if (KeyHandles.Num() > 1)
+		if (bDeleteLastKey == false && KeyHandles.Num() > 1)
 		{
 			OutHandleSet.Add(KeyHandles.Last(0), ECurvePointType::Key);
 		}
 
+		// Below is true except if they are on subframes, we need to delete them
 		// We need to leave the first and the last key of the selection alone for two reasons;
 		// 1. Undo/Redo works better as selections aren't transacted so they don't handle keys being removed/re-added.
 		// by leaving the keys alone, they survive undo/redo operations.
 		// 2. With high interval settings, the shape of the curve can change drastically because the next interval would fall right
 		// outside the last key. If we remove the keys, then the shape changes in these cases. We avoid putting a new key on the last
 		// key (if it exists), to avoid any duplication. This preserves the shape of the curves better in most test cases.
-		if (KeyHandles.Num() > 0)
+		if (bDeleteFirstKey == false && KeyHandles.Num() > 0)
 		{
 			KeyHandles.RemoveAt(0);
 		}
-		if (KeyHandles.Num() > 1)
+		if (bDeleteLastKey == false && KeyHandles.Num() > 1)
 		{
 			KeyHandles.RemoveAt(KeyHandles.Num() - 1);
 		}
 
+
 		// Remove all the old in-between keys and add the new ones
 		Curve->RemoveKeys(KeyHandles);
 		Curve->AddKeys(NewKeyPositions, NewKeyAttributes, &NewKeyHandlesView);
