Curve Editor: Draw a horizontal dotted line per curve when only 1 point is selected on that curve.

diff --git a/Engine/Source/Editor/CurveEditor/Private/CurveEditorSettings.cpp b/Engine/Source/Editor/CurveEditor/Private/CurveEditorSettings.cpp
--- a/Engine/Source/Editor/CurveEditor/Private/CurveEditorSettings.cpp
+++ b/Engine/Source/Editor/CurveEditor/Private/CurveEditorSettings.cpp
@@ -327,7 +327,7 @@ FLinearColor UCurveEditorSettings::GetNextRandomColor()
 	return Color;
 }
 
-void UCurveEditorSettings::SetTreeViewWidth(float InTreeViewWidth)
+void UCurveEditorSettings::SetTreeViewWidth(const float InTreeViewWidth)
 {
 	if (InTreeViewWidth != TreeViewWidth)
 	{
@@ -336,7 +336,7 @@ void UCurveEditorSettings::SetTreeViewWidth(float InTreeViewWidth)
 	}
 }
 
-void UCurveEditorSettings::SetMarqueePointSensitivity(float InMarqueePointSensitivity)
+void UCurveEditorSettings::SetMarqueePointSensitivity(const float InMarqueePointSensitivity)
 {
 	if (InMarqueePointSensitivity != MarqueePointSensitivity)
 	{
@@ -345,6 +345,15 @@ void UCurveEditorSettings::SetMarqueePointSensitivity(float InMarqueePointSensit
 	}
 }
 
+void UCurveEditorSettings::SetMaxNumHorizontalSelectionLinesPerCurve(const float InValue)
+{
+	if (InValue != MaxNumHorizontalSelectionLinesPerCurve)
+	{
+		MaxNumHorizontalSelectionLinesPerCurve = InValue;
+		SaveConfig();
+	}
+}
+
 void UCurveEditorSettings::PostEditChangeProperty(FPropertyChangedEvent& PropertyChangedEvent)
 {
 	FName PropertyName = (PropertyChangedEvent.Property != NULL) ? PropertyChangedEvent.Property->GetFName() : NAME_None;
diff --git a/Engine/Source/Editor/CurveEditor/Private/Views/SInteractiveCurveEditorView.cpp b/Engine/Source/Editor/CurveEditor/Private/Views/SInteractiveCurveEditorView.cpp
--- a/Engine/Source/Editor/CurveEditor/Private/Views/SInteractiveCurveEditorView.cpp
+++ b/Engine/Source/Editor/CurveEditor/Private/Views/SInteractiveCurveEditorView.cpp
@@ -2,6 +2,8 @@
 
 #include "Views/SInteractiveCurveEditorView.h"
 
+#include "Algo/MaxElement.h"
+#include "Algo/RemoveIf.h"
 #include "Algo/Sort.h"
 #include "AnimatedRange.h"
 #include "Containers/ArrayView.h"
@@ -266,11 +268,13 @@ void SInteractiveCurveEditorView::PaintView(const FPaintArgs& Args, const FGeome
 	if (CurveEditor)
 	{
 		const ESlateDrawEffect DrawEffects = ShouldBeEnabled(bParentEnabled) ? ESlateDrawEffect::None : ESlateDrawEffect::DisabledEffect;
-
+		const TSharedRef<FCurveEditor> CurveEditorRef = CurveEditor.ToSharedRef();
+		
 		DrawBackground(AllottedGeometry, OutDrawElements, BaseLayerId, DrawEffects);
-		DrawGridLines(CurveEditor.ToSharedRef(), AllottedGeometry, OutDrawElements, BaseLayerId, DrawEffects);
-		DrawBufferedCurves(CurveEditor.ToSharedRef(), AllottedGeometry, MyCullingRect, OutDrawElements, BaseLayerId, InWidgetStyle, DrawEffects);
-		DrawCurves(CurveEditor.ToSharedRef(), AllottedGeometry, MyCullingRect, OutDrawElements, BaseLayerId, InWidgetStyle, DrawEffects);
+		DrawGridLines(CurveEditorRef, AllottedGeometry, OutDrawElements, BaseLayerId, DrawEffects);
+		DrawBufferedCurves(CurveEditorRef, AllottedGeometry, MyCullingRect, OutDrawElements, BaseLayerId, InWidgetStyle, DrawEffects);
+		DrawCurves(CurveEditorRef, AllottedGeometry, MyCullingRect, OutDrawElements, BaseLayerId, InWidgetStyle, DrawEffects);
+		DrawDottedVisualizationLine(CurveEditorRef, AllottedGeometry, OutDrawElements, BaseLayerId);
 	}
 }
 
@@ -931,6 +935,110 @@ void SInteractiveCurveEditorView::DrawBufferedCurves(TSharedRef<FCurveEditor> Cu
 	}
 }
 
+void SInteractiveCurveEditorView::DrawDottedVisualizationLine(
+	TSharedRef<FCurveEditor> InCurveEditor, const FGeometry& InAllottedGeometry, FSlateWindowElementList& OutDrawElements, const int32 InBaseLayerId
+	) const
+{
+	UpdatedKeysWithDottedLines(*InCurveEditor);
+	
+	const FCurveEditorScreenSpace ViewSpace = GetViewSpace();
+	for (const TPair<FCurveModelID, TArray<FKeyHandle>>& Pair : CurvesToKeysWithDottedLines)
+	{
+		const FCurveModel* CurveModel = InCurveEditor->FindCurve(Pair.Key);
+		if (!CurveModel)
+		{
+			continue;
+		}
+	
+		for (const FKeyHandle& Handle : Pair.Value)
+		{
+			FKeyPosition Position;
+			CurveModel->GetKeyPositions({ Handle }, TArrayView<FKeyPosition>(&Position, 1));
+			
+			const float PosY = ViewSpace.ValueToScreen(Position.OutputValue);
+			FSlateDrawElement::MakeDashedLines(
+				OutDrawElements,
+				InBaseLayerId + CurveViewConstants::ELayerOffset::GridOverlays,
+				InAllottedGeometry.ToPaintGeometry(),
+				{ FVector2f{ 0, PosY }, FVector2f{ InAllottedGeometry.GetLocalSize().X, PosY } }
+				);
+		}
+	}
+}
+
+void SInteractiveCurveEditorView::UpdatedKeysWithDottedLines(const FCurveEditor& InCurveEditor) const
+{
+	// We'll draw a horizontal line for each selected point on each selected curve. By default, at most 1 line per curve with selected points.
+	const int32 MaxLinesPerCurve = InCurveEditor.GetSettings()->GetMaxNumHorizontalSelectionLinesPerCurve();
+
+	const TMap<FCurveModelID, FKeyHandleSet>& Selection = InCurveEditor.Selection.GetAll();
+	for (const TPair<FCurveModelID, FKeyHandleSet>& Pair : Selection)
+	{
+		const FCurveModelID& CurveId = Pair.Key;
+		const FCurveModel* CurveModel = InCurveEditor.FindCurve(CurveId);
+		if (!CurveModel)
+		{
+			continue;
+		}
+
+		// If there are multiple points selected, draw dotted lines through the same points every tick until the user deselects them.
+		const FKeyHandleSet& KeyHandleSet = Pair.Value;
+		TArray<FKeyHandle>& KeysWithDottedLines = CurvesToKeysWithDottedLines.FindOrAdd(CurveId);
+		const int32 NumToSet = Algo::RemoveIf(KeysWithDottedLines, [&KeyHandleSet](const FKeyHandle& Handle)
+		{
+			return !KeyHandleSet.Contains(Handle, ECurvePointType::Key);
+		});
+		KeysWithDottedLines.SetNum(
+			// User may have changed the max number last tick
+			FMath::Min(NumToSet, MaxLinesPerCurve)
+			);
+
+		// Can add more dotted lines?
+		PickPointsToPlaceDottedLineOn(MaxLinesPerCurve, *CurveModel, KeyHandleSet, KeysWithDottedLines);
+	}
+
+	// Remove dotted lines from curves that are no longer selected.
+	for (auto DottedLineIt = CurvesToKeysWithDottedLines.CreateIterator(); DottedLineIt; ++DottedLineIt)
+	{
+		if (!Selection.Contains(DottedLineIt->Key))
+		{
+			DottedLineIt.RemoveCurrent();
+		}
+	}
+}
+
+void SInteractiveCurveEditorView::PickPointsToPlaceDottedLineOn(
+	const int32 InMaxLinesPerCurve,
+	const FCurveModel& InCurveModel,
+	const FKeyHandleSet& InUserSelectedKeys,
+	TArray<FKeyHandle>& OutKeysWithDottedLines
+	) const
+{
+	// Prefer to pick the key with the highest peak.
+	if (InUserSelectedKeys.Num() > 0 && OutKeysWithDottedLines.Num() < InMaxLinesPerCurve)
+	{
+		const TConstArrayView<FKeyHandle> Handles = InUserSelectedKeys.AsArray();
+		const FKeyHandle* PeakPoint = Algo::MaxElementBy(Handles, [&InCurveModel](const FKeyHandle& Handle)
+		{
+			FKeyPosition Position;
+			InCurveModel.GetKeyPositions({ Handle }, TArrayView<FKeyPosition>(&Position, 1));
+			return Position.OutputValue;
+		});
+		OutKeysWithDottedLines.AddUnique(*PeakPoint);
+	}
+
+	// Can still add more? Select a 'random' key. 
+	for (int32 Index = 0; Index < InUserSelectedKeys.Num() && OutKeysWithDottedLines.Num() < InMaxLinesPerCurve; Index++)
+	{
+		const FKeyHandle& Handle = InUserSelectedKeys.AsArray()[Index];
+		const bool bIsPoint = InUserSelectedKeys.PointType(Handle) == ECurvePointType::Key;
+		if (bIsPoint)
+		{
+			OutKeysWithDottedLines.AddUnique(Handle);
+		}
+	}
+}
+
 bool SInteractiveCurveEditorView::GetPointsWithinWidgetRange(const FSlateRect& WidgetRectangle, TArray<FCurvePointHandle>* OutPoints) const
 {
 	TSharedPtr<FCurveEditor> CurveEditor = WeakCurveEditor.Pin();
diff --git a/Engine/Source/Editor/CurveEditor/Public/CurveEditorSettings.h b/Engine/Source/Editor/CurveEditor/Public/CurveEditorSettings.h
--- a/Engine/Source/Editor/CurveEditor/Public/CurveEditorSettings.h
+++ b/Engine/Source/Editor/CurveEditor/Public/CurveEditorSettings.h
@@ -185,6 +185,11 @@ public:
 	/** Sets how sensitive the selection marquee should be when selecting points. */
 	void SetMarqueePointSensitivity(float InMarqueePointSensitivity);
 
+	/** Gets how sensitive the selection marquee should be when selecting points. */
+	float GetMaxNumHorizontalSelectionLinesPerCurve() const { return MaxNumHorizontalSelectionLinesPerCurve; }
+	/** Sets how sensitive the selection marquee should be when selecting points. */
+	void SetMaxNumHorizontalSelectionLinesPerCurve(float InValue);
+
 protected:
 
 	virtual void PostEditChangeProperty(FPropertyChangedEvent& PropertyChangedEvent) override;
@@ -249,6 +254,9 @@ protected:
 	UPROPERTY(config, EditAnywhere, Category = "Curve Editor", meta = (ClampMin = 0.0, ClampMax = 1.0))
 	float MarqueePointSensitivity = 0.45f;
 
+	UPROPERTY(config, EditAnywhere, Category = "Curve Editor", meta = (ClampMin = 0))
+	int32 MaxNumHorizontalSelectionLinesPerCurve = 1;
+
 private:
 	FOnCustomColorsChanged OnCustomColorsChangedEvent;
 	FOnAxisSnappingChanged OnAxisSnappingChangedEvent;
diff --git a/Engine/Source/Editor/CurveEditor/Public/Views/SInteractiveCurveEditorView.h b/Engine/Source/Editor/CurveEditor/Public/Views/SInteractiveCurveEditorView.h
--- a/Engine/Source/Editor/CurveEditor/Public/Views/SInteractiveCurveEditorView.h
+++ b/Engine/Source/Editor/CurveEditor/Public/Views/SInteractiveCurveEditorView.h
@@ -43,6 +43,7 @@ struct FGeometry;
 struct FOptionalSize;
 struct FPointerEvent;
 struct FKeyAttributes;
+struct FKeyHandleSet;
 
 namespace CurveViewConstants
 {
@@ -138,7 +139,13 @@ protected:
 	void DrawGridLines(TSharedRef<FCurveEditor> CurveEditor, const FGeometry& AllottedGeometry, FSlateWindowElementList& OutDrawElements, int32 BaseLayerId, ESlateDrawEffect DrawEffects) const;
 	void DrawCurves(TSharedRef<FCurveEditor> CurveEditor, const FGeometry& AllottedGeometry, const FSlateRect& MyCullingRect, FSlateWindowElementList& OutDrawElements, int32 BaseLayerId, const FWidgetStyle& InWidgetStyle, ESlateDrawEffect DrawEffects) const;
 	void DrawBufferedCurves(TSharedRef<FCurveEditor> CurveEditor, const FGeometry& AllottedGeometry, const FSlateRect& MyCullingRect, FSlateWindowElementList& OutDrawElements, int32 BaseLayerId, const FWidgetStyle& InWidgetStyle, ESlateDrawEffect DrawEffects) const;
+	void DrawDottedVisualizationLine(TSharedRef<FCurveEditor> InCurveEditor, const FGeometry& InAllottedGeometry, FSlateWindowElementList& OutDrawElements, int32 InBaseLayerId) const;
 
+	/** Updates the keys that will be drawn with dotted lines. */
+	void UpdatedKeysWithDottedLines(const FCurveEditor& InCurveEditor) const;
+	/** Picks which keys will have dotted lines drawn for them. By default, each curve picks the key in the user selection with the highest peak. */
+	virtual void PickPointsToPlaceDottedLineOn(const int32 InMaxLinesPerCurve, const FCurveModel& InCurveModel, const FKeyHandleSet& InUserSelectedKeys, TArray<FKeyHandle>& OutKeysWithDottedLines) const;
+	
 	FSlateColor GetCurveCaptionColor() const;
 	FText GetCurveCaption() const;
 
@@ -197,6 +204,15 @@ protected:
 	/** Background tint for this widget */
 	FLinearColor BackgroundTint;
 
+	/**
+	 * We draw a horizontal dotted line on the keys that a user selects.
+	 * This makes it easier for the user to see a selected key's value in relation to other keys.
+	 * The  per-curve upper limit of dotted lines that are drawn is configured in UCurveEditorSettings::MaxNumHorizontalSelectionLinesPerCurve.
+	 *
+	 * This is called in OnPaint, which is const. Mutable is a workaround.
+	 */
+	mutable TMap<FCurveModelID, TArray<FKeyHandle>> CurvesToKeysWithDottedLines;
+
 private:
 
 	/** (Optional) the current drag operation */
