Curve Editor: No longer able to select both keys and bezier handles with the marquee.

When pressing a modifier key,
- If selection contains only keys, marquee will only append keys.
- If selection contains only tangents, marquee will only append tangents.
- Otherwise, if due to a future change both keys and tangents are in the selection, append both to selection.

diff --git a/Engine/Source/Editor/CurveEditor/Private/DragOperations/CurveEditorDragOperation_Marquee.cpp b/Engine/Source/Editor/CurveEditor/Private/DragOperations/CurveEditorDragOperation_Marquee.cpp
--- a/Engine/Source/Editor/CurveEditor/Private/DragOperations/CurveEditorDragOperation_Marquee.cpp
+++ b/Engine/Source/Editor/CurveEditor/Private/DragOperations/CurveEditorDragOperation_Marquee.cpp
@@ -98,6 +98,56 @@ static void RestrictSelectionToEitherPointsOrTangents(bool bPreferPointSelection
 		CurveEditor->Selection.Remove(Point);
 	}
 }
+
+struct FSelectionPreferences
+{
+	bool bPreferPointSelection;
+	bool bPreferTangentSelection;
+};
+static FSelectionPreferences GetSelectionPreferences(
+	const TArray<FCurvePointHandle>& AllPoints, const FCurveEditor* CurveEditor, const bool bAnyModifiersKeysDown
+	)
+{
+	const bool bSelectionContainsKeys = Algo::AnyOf(CurveEditor->Selection.GetAll(), [](const TPair<FCurveModelID, FKeyHandleSet>& Pair)
+	{
+		const FKeyHandleSet& HandleSet = Pair.Value;
+		return Algo::AnyOf(Pair.Value.AsArray(), [&HandleSet](const FKeyHandle& Handle)
+			{ return HandleSet.PointType(Handle) == ECurvePointType::Key; }
+		);
+	});
+	const bool bSelectionContainsTangents = Algo::AnyOf(CurveEditor->Selection.GetAll(), [](const TPair<FCurveModelID, FKeyHandleSet>& Pair)
+	{
+		const FKeyHandleSet& HandleSet = Pair.Value;
+		return Algo::AnyOf(Pair.Value.AsArray(), [&HandleSet](const FKeyHandle& Handle)
+		{
+			const ECurvePointType PointType = HandleSet.PointType(Handle);
+			return PointType == ECurvePointType::ArriveTangent || PointType == ECurvePointType::LeaveTangent;
+		});
+	});
+	// If, for whatever reason, the selection contains both keys and tangents already, prefer neither of the other.
+	if (bSelectionContainsKeys && bSelectionContainsTangents)
+	{
+		return { false, false };
+	}
+	
+	const bool bMarqueeHasAnyKeys = Algo::AnyOf(AllPoints, [](const FCurvePointHandle& Point)
+		{ return Point.PointType == ECurvePointType::Key; }
+	);
+	
+	// If there are any points to be selected, prefer selecting points over tangents.
+	const bool bPreferPointSelection = (bMarqueeHasAnyKeys && !bAnyModifiersKeysDown)
+		// If the selection already contains keys, also prefer selecting keys, i.e. don't add tangents to the pre-existing key selection... [1]
+		|| (bAnyModifiersKeysDown && bSelectionContainsKeys)
+		// If modifying the selection, prefer keys only if the selection does not contain tangents.
+		|| (bMarqueeHasAnyKeys && bAnyModifiersKeysDown && !bSelectionContainsTangents);
+
+	// We'll prefer tangents when the selection contains only tangents and modifying the selection (i.e. a modifier key is pressed).
+	const bool bPreferTangentSelection = !bPreferPointSelection
+		// [1] ... likewise if modifying the selection, prefer selecting tangents
+		&& bAnyModifiersKeysDown && bSelectionContainsTangents;
+
+	return { bPreferPointSelection, bPreferTangentSelection };
+}
 }
 
 FCurveEditorDragOperation_Marquee::FCurveEditorDragOperation_Marquee(FCurveEditor* InCurveEditor)
@@ -141,17 +191,16 @@ void FCurveEditorDragOperation_Marquee::OnEndDrag(FVector2D InitialPosition, FVe
 	const bool bIsShiftDown = MouseEvent.IsShiftDown();
 	const bool bIsAltDown = MouseEvent.IsAltDown();
 	const bool bIsControlDown = MouseEvent.IsControlDown();
-
-	if (!bIsShiftDown && !bIsAltDown && !bIsControlDown)
+	const bool bAnyModifiersKeysDown = bIsShiftDown || bIsAltDown || bIsControlDown;
+	
+	if (!bAnyModifiersKeysDown)
 	{
 		CurveEditor->Selection.Clear();
 	}
-
-	// If there are any points to be selected, prefer selecting points over tangents
-	const bool bPreferPointSelection = Algo::AnyOf(AllPoints, [](const FCurvePointHandle& Point)
-	{
-		return Point.PointType == ECurvePointType::Key;
-	});
+	
+	const auto[bPreferPointSelection, bPreferTangentSelection] = UE::CurveEditor::MarqueeDetail::GetSelectionPreferences(
+		AllPoints, CurveEditor, bAnyModifiersKeysDown
+		);
 
 	// When adding to the existing selection, ensure only either points or tangents are selected
 	if (bIsShiftDown)
@@ -187,6 +236,13 @@ void FCurveEditorDragOperation_Marquee::OnEndDrag(FVector2D InitialPosition, FVe
 				CurveEditor->Selection.Add(Point);
 			}
 		}
+		else if (bPreferTangentSelection)
+		{
+			if (Point.PointType == ECurvePointType::ArriveTangent || Point.PointType == ECurvePointType::LeaveTangent)
+			{
+				CurveEditor->Selection.Add(Point);
+			}
+		}
 		else
 		{
 			CurveEditor->Selection.Add(Point);
