Curve Editor: Allow shift to override axis snapping.

diff --git a/Engine/Source/Editor/CurveEditor/Public/CurveEditorSnapMetrics.h b/Engine/Source/Editor/CurveEditor/Public/CurveEditorSnapMetrics.h
--- a/Engine/Source/Editor/CurveEditor/Public/CurveEditorSnapMetrics.h
+++ b/Engine/Source/Editor/CurveEditor/Public/CurveEditorSnapMetrics.h
@@ -82,61 +82,70 @@ struct FCurveEditorAxisSnap
 		RestrictedAxisList = ECurveEditorSnapAxis::CESA_None;
 	}
 
-	FVector2D GetSnappedPosition(FVector2D InitialPosition, FVector2D LastPosition, FVector2D CurrentPosition, const FPointerEvent& MouseEvent, FSnapState& InOutSnapState, const bool bIgnoreAxisLock = false)
+	/**
+	 * Combines an InitialPosition and mouse movement to produce a final position that respects the axis snapping settings.
+	 * Pressing shift ignores the snapping settings.
+	 * 
+	 * For example, if movement is constrained to x-axis only and the mouse moves in direction FVector2D{ 100, 200 }, only the delta movement of
+	 * FVector2D{ 100, 0} is applied.
+	 * 
+	 * @return The end position resulting from applying the snapping behavior to the mouse movement.
+	 */
+	FVector2D GetSnappedPosition(
+		const FVector2D& InitialPosition,
+		const FVector2D& LastPosition,
+		const FVector2D& CurrentPosition,
+		const FPointerEvent& MouseEvent,
+		FSnapState& InOutSnapState,
+		const bool bIgnoreAxisLock = false
+		)
 	{
-		// If we're ignoring axis lock (such as for UI) we allow them to use shift anyways
-		bool bCanUseShift = bIgnoreAxisLock || RestrictedAxisList == ECurveEditorSnapAxis::CESA_None;
-
 		FVector2D MouseLockVector = FVector2D::UnitVector;
-		if (bCanUseShift)
+
+		if (MouseEvent.IsShiftDown())
 		{
-			// Snapping works by determining which direction the mouse is moving in while shift is being held
-			// down. Once you have moved far enough to pass the threshold, you're locked into that direction
-			// but if you release shift, it will reset the direction and allow you to pick a new direction.
-			if (MouseEvent.IsShiftDown())
+			if (!InOutSnapState.bHasStartPosition)
 			{
-				if (!InOutSnapState.bHasStartPosition)
-				{
-					InOutSnapState.MousePosOnShiftStart = LastPosition;
-					InOutSnapState.bHasStartPosition = true;
-				}
-				// If they have passed the threshold they should have a lock vector they're snapped to.
-				if (!InOutSnapState.bHasPassedThreshold)
-				{
-					InOutSnapState.MouseLockVector = MouseLockVector;
+				InOutSnapState.MousePosOnShiftStart = LastPosition;
+				InOutSnapState.bHasStartPosition = true;
+			}
+			// If they have passed the threshold they should have a lock vector they're snapped to.
+			if (!InOutSnapState.bHasPassedThreshold)
+			{
+				InOutSnapState.MouseLockVector = MouseLockVector;
 					
-					// They have not passed the threshold yet, let's see if they've passed it now.
-					FVector2D DragDelta = CurrentPosition - InOutSnapState.MousePosOnShiftStart;
-					if (DragDelta.Size() > 0.001f)
+				// They have not passed the threshold yet, let's see if they've passed it now.
+				const FVector2D DragDelta = CurrentPosition - InOutSnapState.MousePosOnShiftStart;
+				if (DragDelta.Size() > 0.001f)
+				{
+					InOutSnapState.bHasPassedThreshold = true;
+					InOutSnapState.MouseLockVector = FVector2D::UnitVector;
+					if (FMath::Abs(DragDelta.X) > FMath::Abs(DragDelta.Y))
+					{
+						InOutSnapState.MouseLockVector.Y = 0;
+					}
+					else
 					{
-						InOutSnapState.bHasPassedThreshold = true;
-						InOutSnapState.MouseLockVector = FVector2D::UnitVector;
-						if (FMath::Abs(DragDelta.X) > FMath::Abs(DragDelta.Y))
-						{
-							InOutSnapState.MouseLockVector.Y = 0;
-						}
-						else
-						{
-							InOutSnapState.MouseLockVector.X = 0;
-						}
+						InOutSnapState.MouseLockVector.X = 0;
 					}
 				}
 			}
-			else
-			{
-				// If they don't have shift pressed just disable the lock.
-				InOutSnapState.Reset();
-			}
-
+			
 			MouseLockVector = InOutSnapState.MouseLockVector;
 		}
-		else if (!bIgnoreAxisLock && RestrictedAxisList == ECurveEditorSnapAxis::CESA_X)
-		{
-			MouseLockVector.Y = 0;
-		}
-		else if (!bIgnoreAxisLock && RestrictedAxisList == ECurveEditorSnapAxis::CESA_Y)
+		else
 		{
-			MouseLockVector.X = 0;
+			// If they don't have shift pressed anymore just disable the lock.
+			InOutSnapState.Reset();
+
+			if (!bIgnoreAxisLock && RestrictedAxisList == ECurveEditorSnapAxis::CESA_X)
+			{
+				MouseLockVector.Y = 0;
+			}
+			else if (!bIgnoreAxisLock && RestrictedAxisList == ECurveEditorSnapAxis::CESA_Y)
+			{
+				MouseLockVector.X = 0;
+			}
 		}
 
 		return InitialPosition + (CurrentPosition - InitialPosition) * MouseLockVector;
