Curve Editor: Fix crash with smart reduce filter

diff --git a/Engine/Source/Editor/CurveEditor/Private/Filters/CurveEditorSmartReduceFilter.cpp b/Engine/Source/Editor/CurveEditor/Private/Filters/CurveEditorSmartReduceFilter.cpp
--- a/Engine/Source/Editor/CurveEditor/Private/Filters/CurveEditorSmartReduceFilter.cpp
+++ b/Engine/Source/Editor/CurveEditor/Private/Filters/CurveEditorSmartReduceFilter.cpp
@@ -232,7 +232,7 @@ void UCurveEditorSmartReduceFilter::SmartReduce(FCurveModel* Curve, const FSmart
 		bool bKeyWasAdded = false;
 		if (KeyIndex < KeyPosAttrs.Num())
 		{
-			while ((KeyTimes[KeyIndex] <= (Cache.Times[Index] + UE_DOUBLE_SMALL_NUMBER)) && KeyIndex < KeyTimes.Num())
+			while (KeyIndex < KeyTimes.Num() && (KeyTimes[KeyIndex] <= (Cache.Times[Index] + UE_DOUBLE_SMALL_NUMBER)))
 			{
 				KeyPosition.InputValue = KeyPosAttrs[KeyIndex].Key.InputValue;
 				KeyPosition.OutputValue = KeyPosAttrs[KeyIndex].Key.OutputValue;
