Curve Editor: Fix pasting auto-smart keys in rich curve editor changing their interpolation types.

diff --git a/Engine/Source/Editor/CurveEditor/Private/RichCurveEditorModel.cpp b/Engine/Source/Editor/CurveEditor/Private/RichCurveEditorModel.cpp
--- a/Engine/Source/Editor/CurveEditor/Private/RichCurveEditorModel.cpp
+++ b/Engine/Source/Editor/CurveEditor/Private/RichCurveEditorModel.cpp
@@ -31,6 +31,86 @@
 
 class FCurveEditor;
 
+namespace UE::CurveEditor::RichModelDetail
+{
+static bool IsAuto(ERichCurveTangentMode TangentMode)
+{
+	return (TangentMode == RCTM_Auto || TangentMode == RCTM_SmartAuto);
+}
+
+/**
+ * Shared implementation for getting keys.
+ *
+ * Typically, attributes reflect the settings that the user has manually configured for the keys, so certain attributes may remain unset.
+ * For instance, when TangentMode == RCTM_Auto, tangents and weights are automatically computed, meaning attributes like ArriveTangent are not explicitly set.
+ * Setting ArriveTangent would imply a user-defined value, which is incompatible with TangentMode == RCTM_Auto.
+ *
+ * Sometimes API users still need a way to get the auto-computed values, like those of tangents, though.
+ * For this, use bAllAttributes parameter:
+ * - true: get all attributes, even the auto-computed ones. The values are useful for e.g. UI visualization. Do not pass them to SetKeyAttributes.
+ * - false: get only the attributes that were set by the user (i.e. skip returning auto-computed values). You can pass them to SetKeyAttributes, e.g. for copy paste.
+ */
+template<bool bAllAttributes>
+static void GetKeyAttributes(
+	const FRichCurve& InRichCurve,
+	TArrayView<const FKeyHandle> InKeys,
+	TArrayView<FKeyAttributes> OutAttributes
+	)
+{
+	const TArray<FRichCurveKey>& AllKeys = InRichCurve.GetConstRefOfKeys();
+	if (AllKeys.Num() == 0)
+	{
+		return;
+	}
+
+	for (int32 Index = 0; Index < InKeys.Num(); ++Index)
+	{
+		if (InRichCurve.IsKeyHandleValid(InKeys[Index]))
+		{
+			const FRichCurveKey& ThisKey    = InRichCurve.GetKeyRef(InKeys[Index]);
+			FKeyAttributes&      Attributes = OutAttributes[Index];
+
+			Attributes.SetInterpMode(ThisKey.InterpMode);
+
+			// If the previous key is cubic, show the arrive tangent handle even if this key is constant
+			const FKeyHandle PreviousKeyHandle = InRichCurve.GetPreviousKey(InKeys[Index]);
+			const bool bGetArriveTangent = bAllAttributes
+				&& InRichCurve.IsKeyHandleValid(PreviousKeyHandle) && InRichCurve.GetKeyRef(PreviousKeyHandle).InterpMode == RCIM_Cubic;
+			if (bGetArriveTangent)
+			{
+				Attributes.SetArriveTangent(ThisKey.ArriveTangent);
+			}
+
+			if (ThisKey.InterpMode != RCIM_Constant && ThisKey.InterpMode != RCIM_Linear)
+			{
+				Attributes.SetTangentMode(ThisKey.TangentMode);
+
+				// The remaining settings (arrive / leave tangent, arrive / leave weight, weight mode) can only be user specified, i.e. cannot be set when in auto.
+				// Setting any of them would imply they specified by the user, which would contradict the setting that it's auto. 
+				if (!bAllAttributes && IsAuto(ThisKey.TangentMode))
+				{
+					continue;
+				}
+				
+				Attributes.SetArriveTangent(ThisKey.ArriveTangent);
+				Attributes.SetLeaveTangent(ThisKey.LeaveTangent);
+
+				if (ThisKey.InterpMode == RCIM_Cubic)
+				{
+					Attributes.SetTangentWeightMode(ThisKey.TangentWeightMode);
+					if (ThisKey.TangentWeightMode != RCTWM_WeightedNone)
+					{
+						Attributes.SetArriveTangentWeight(ThisKey.ArriveTangentWeight);
+						Attributes.SetLeaveTangentWeight(ThisKey.LeaveTangentWeight);
+					}
+				}
+			}
+		}
+	}
+	
+}
+}
+
 
 void RefineCurvePoints(const FRichCurve& RichCurve, double TimeThreshold, float ValueThreshold, TArray<TTuple<double, double>>& InOutPoints)
 {
@@ -365,53 +445,19 @@ void FRichCurveEditorModel::SetKeyPositions(TArrayView<const FKeyHandle> InKeys,
 
 void FRichCurveEditorModel::GetKeyAttributes(TArrayView<const FKeyHandle> InKeys, TArrayView<FKeyAttributes> OutAttributes) const
 {
-	if (UObject* Owner = WeakOwner.Get())
+	if (UObject* Owner = WeakOwner.Get(); Owner && IsValid())
 	{
-		if(IsValid())
-		{
-			const FRichCurve& RichCurve = GetReadOnlyRichCurve();
-			const TArray<FRichCurveKey>& AllKeys = RichCurve.GetConstRefOfKeys();
-			if (AllKeys.Num() == 0)
-			{
-				return;
-			}
-
-			for (int32 Index = 0; Index < InKeys.Num(); ++Index)
-			{
-				if (RichCurve.IsKeyHandleValid(InKeys[Index]))
-				{
-					const FRichCurveKey& ThisKey    = RichCurve.GetKeyRef(InKeys[Index]);
-					FKeyAttributes&      Attributes = OutAttributes[Index];
-
-					Attributes.SetInterpMode(ThisKey.InterpMode);
-
-					// If the previous key is cubic, show the arrive tangent handle even if this key is constant
-					FKeyHandle PreviousKeyHandle = RichCurve.GetPreviousKey(InKeys[Index]);
-					const bool bGetArriveTangent = RichCurve.IsKeyHandleValid(PreviousKeyHandle) && RichCurve.GetKeyRef(PreviousKeyHandle).InterpMode == RCIM_Cubic;
-					if (bGetArriveTangent)
-					{
-						Attributes.SetArriveTangent(ThisKey.ArriveTangent);
-					}
-
-					if (ThisKey.InterpMode != RCIM_Constant && ThisKey.InterpMode != RCIM_Linear)
-					{
-						Attributes.SetTangentMode(ThisKey.TangentMode);
-						Attributes.SetArriveTangent(ThisKey.ArriveTangent);
-						Attributes.SetLeaveTangent(ThisKey.LeaveTangent);
+		const FRichCurve& RichCurve = GetReadOnlyRichCurve();
+		UE::CurveEditor::RichModelDetail::GetKeyAttributes<true>(RichCurve, InKeys, OutAttributes);
+	}
+}
 
-						if (ThisKey.InterpMode == RCIM_Cubic)
-						{
-							Attributes.SetTangentWeightMode(ThisKey.TangentWeightMode);
-							if (ThisKey.TangentWeightMode != RCTWM_WeightedNone)
-							{
-								Attributes.SetArriveTangentWeight(ThisKey.ArriveTangentWeight);
-								Attributes.SetLeaveTangentWeight(ThisKey.LeaveTangentWeight);
-							}
-						}
-					}
-				}
-			}
-		}
+void FRichCurveEditorModel::GetKeyAttributesExcludingAutoComputed(TArrayView<const FKeyHandle> InKeys, TArrayView<FKeyAttributes> OutAttributes) const
+{
+	if (UObject* Owner = WeakOwner.Get(); Owner && IsValid())
+	{
+		const FRichCurve& RichCurve = GetReadOnlyRichCurve();
+		UE::CurveEditor::RichModelDetail::GetKeyAttributes<false>(RichCurve, InKeys, OutAttributes);
 	}
 }
 
diff --git a/Engine/Source/Editor/CurveEditor/Public/RichCurveEditorModel.h b/Engine/Source/Editor/CurveEditor/Public/RichCurveEditorModel.h
--- a/Engine/Source/Editor/CurveEditor/Public/RichCurveEditorModel.h
+++ b/Engine/Source/Editor/CurveEditor/Public/RichCurveEditorModel.h
@@ -45,6 +45,7 @@ public:
 	virtual void SetKeyPositions(TArrayView<const FKeyHandle> InKeys, TArrayView<const FKeyPosition> InKeyPositions, EPropertyChangeType::Type ChangeType) override;
 
 	virtual void GetKeyAttributes(TArrayView<const FKeyHandle> InKeys, TArrayView<FKeyAttributes> OutAttributes) const override;
+	virtual void GetKeyAttributesExcludingAutoComputed(TArrayView<const FKeyHandle> InKeys, TArrayView<FKeyAttributes> OutAttributes) const override;
 	virtual void SetKeyAttributes(TArrayView<const FKeyHandle> InKeys, TArrayView<const FKeyAttributes> InAttributes, EPropertyChangeType::Type ChangeType = EPropertyChangeType::Unspecified) override;
 
 	virtual void GetCurveAttributes(FCurveAttributes& OutCurveAttributes) const override;
