Curve Editor: Dotted value indicator line adjustments:
- Draw dotted line only when exactly 1 curve is selected
- Draw up to 2 dotted lines: one for the min and max selected key, respectively.
- Made the curve slightly less bright
- Can be toggled on and off in project settings and using the eye icon in curve editor

diff --git a/Engine/Source/Editor/CurveEditor/Private/CurveEditor.cpp b/Engine/Source/Editor/CurveEditor/Private/CurveEditor.cpp
--- a/Engine/Source/Editor/CurveEditor/Private/CurveEditor.cpp
+++ b/Engine/Source/Editor/CurveEditor/Private/CurveEditor.cpp
@@ -524,12 +524,19 @@ void FCurveEditor::BindCommands()
 	CommandList->MapAction(FCurveEditorCommands::Get().ToggleShowBufferedCurves,
 		FExecuteAction::CreateLambda( [CurveSettings]{ CurveSettings->SetShowBufferedCurves( !CurveSettings->GetShowBufferedCurves() ); } ),
 		FCanExecuteAction(),
-		FIsActionChecked::CreateLambda( [CurveSettings]{ return CurveSettings->GetShowBufferedCurves(); } ) );
+		FIsActionChecked::CreateLambda( [CurveSettings]{ return CurveSettings->GetShowBufferedCurves(); } )
+		);
 
 	CommandList->MapAction(FCurveEditorCommands::Get().ToggleShowCurveEditorCurveToolTips,
 		FExecuteAction::CreateLambda( [CurveSettings]{ CurveSettings->SetShowCurveEditorCurveToolTips( !CurveSettings->GetShowCurveEditorCurveToolTips() ); } ),
 		FCanExecuteAction(),
-		FIsActionChecked::CreateLambda( [CurveSettings]{ return CurveSettings->GetShowCurveEditorCurveToolTips(); } ) );
+		FIsActionChecked::CreateLambda( [CurveSettings]{ return CurveSettings->GetShowCurveEditorCurveToolTips(); } )
+		);
+	CommandList->MapAction(FCurveEditorCommands::Get().ToggleShowValueIndicatorLines,
+		FExecuteAction::CreateLambda( [CurveSettings]{ CurveSettings->SetShowValueIndicators( !CurveSettings->GetShowValueIndicators() ); } ),
+		FCanExecuteAction(),
+		FIsActionChecked::CreateLambda( [CurveSettings]{ return CurveSettings->GetShowValueIndicators(); } )
+	);
 
 	// Deactivate Current Tool
 	CommandList->MapAction(FCurveEditorCommands::Get().DeactivateCurrentTool,
diff --git a/Engine/Source/Editor/CurveEditor/Private/CurveEditorCommands.cpp b/Engine/Source/Editor/CurveEditor/Private/CurveEditorCommands.cpp
--- a/Engine/Source/Editor/CurveEditor/Private/CurveEditorCommands.cpp
+++ b/Engine/Source/Editor/CurveEditor/Private/CurveEditorCommands.cpp
@@ -63,6 +63,7 @@ void FCurveEditorCommands::RegisterCommands()
 	UI_COMMAND(ToggleShowBufferedCurves, "Buffered Curves", "Show buffered curves for the selected curves.", EUserInterfaceActionType::ToggleButton, FInputChord() );
 	UI_COMMAND(ToggleShowCurveEditorCurveToolTips, "Curve Tool Tips", "Show a tool tip with name and values when hovering over a curve.", EUserInterfaceActionType::ToggleButton, FInputChord() );
 	UI_COMMAND(ToggleShowBars, "Show Bars", "Show Bars like Constraints and Spaces", EUserInterfaceActionType::ToggleButton, FInputChord());
+	UI_COMMAND(ToggleShowValueIndicatorLines, "Show Value Indicators", "Show horizontal, dotted lines on the selected min & max valued key handles", EUserInterfaceActionType::ToggleButton, FInputChord());
 
 	UI_COMMAND(AddKeyHovered, "Add Key", "Add a new key to this curve at the current position.", EUserInterfaceActionType::Button, FInputChord(EKeys::MiddleMouseButton) );
 	UI_COMMAND(PasteKeysHovered, "Paste", "Paste clipboard contents", EUserInterfaceActionType::Button, FInputChord(EModifierKey::Control, EKeys::V));
diff --git a/Engine/Source/Editor/CurveEditor/Private/CurveEditorSettings.cpp b/Engine/Source/Editor/CurveEditor/Private/CurveEditorSettings.cpp
--- a/Engine/Source/Editor/CurveEditor/Private/CurveEditorSettings.cpp
+++ b/Engine/Source/Editor/CurveEditor/Private/CurveEditorSettings.cpp
@@ -345,11 +345,12 @@ void UCurveEditorSettings::SetMarqueePointSensitivity(const float InMarqueePoint
 	}
 }
 
-void UCurveEditorSettings::SetMaxNumHorizontalSelectionLinesPerCurve(const float InValue)
+void UCurveEditorSettings::SetShowValueIndicators(bool bValue)
 {
-	if (InValue != MaxNumHorizontalSelectionLinesPerCurve)
+	if (bValue != bShowValueIndicators)
 	{
-		MaxNumHorizontalSelectionLinesPerCurve = InValue;
+		bShowValueIndicators = bValue;
+		OnShowValueIndicatorsChangedEvent.Broadcast();
 		SaveConfig();
 	}
 }
diff --git a/Engine/Source/Editor/CurveEditor/Private/SCurveEditorPanel.cpp b/Engine/Source/Editor/CurveEditor/Private/SCurveEditorPanel.cpp
--- a/Engine/Source/Editor/CurveEditor/Private/SCurveEditorPanel.cpp
+++ b/Engine/Source/Editor/CurveEditor/Private/SCurveEditorPanel.cpp
@@ -1044,6 +1044,7 @@ TSharedRef<SWidget> SCurveEditorPanel::MakeCurveEditorCurveViewOptionsMenu()
 	MenuBuilder.AddMenuEntry(FCurveEditorCommands::Get().ToggleShowBufferedCurves);
 	MenuBuilder.AddMenuEntry(FCurveEditorCommands::Get().ToggleShowBars);
 	MenuBuilder.AddMenuEntry(FCurveEditorCommands::Get().ToggleShowCurveEditorCurveToolTips);
+	MenuBuilder.AddMenuEntry(FCurveEditorCommands::Get().ToggleShowValueIndicatorLines);
 
 	MenuBuilder.BeginSection("Organize", LOCTEXT("CurveEditorMenuOrganizeHeader", "Organize"));
 	{
diff --git a/Engine/Source/Editor/CurveEditor/Private/Views/SInteractiveCurveEditorView.cpp b/Engine/Source/Editor/CurveEditor/Private/Views/SInteractiveCurveEditorView.cpp
--- a/Engine/Source/Editor/CurveEditor/Private/Views/SInteractiveCurveEditorView.cpp
+++ b/Engine/Source/Editor/CurveEditor/Private/Views/SInteractiveCurveEditorView.cpp
@@ -2,8 +2,8 @@
 
 #include "Views/SInteractiveCurveEditorView.h"
 
+#include "Algo/AnyOf.h"
 #include "Algo/MaxElement.h"
-#include "Algo/RemoveIf.h"
 #include "Algo/Sort.h"
 #include "AnimatedRange.h"
 #include "Containers/ArrayView.h"
@@ -118,9 +118,13 @@ void SInteractiveCurveEditorView::Construct(const FArguments& InArgs, TWeakPtr<F
 	bAutoSize = InArgs._AutoSize;
 
 	WeakCurveEditor = InCurveEditor;
-	CachedValues.CachedTangentVisibility = InCurveEditor.Pin()->GetSettings()->GetTangentVisibility();
-
-	InCurveEditor.Pin()->OnActiveToolChangedDelegate.AddSP(this, &SInteractiveCurveEditorView::OnCurveEditorToolChanged);
+	{
+		TSharedPtr<FCurveEditor> CurveEditorPin = InCurveEditor.Pin();
+		check(CurveEditorPin);
+		CachedValues.CachedTangentVisibility = CurveEditorPin->GetSettings()->GetTangentVisibility();
+		CurveEditorPin->OnActiveToolChangedDelegate.AddSP(this, &SInteractiveCurveEditorView::OnCurveEditorToolChanged);
+		CurveEditorPin->GetSettings()->OnShowValueIndicatorsChanged().AddSP(this, &SInteractiveCurveEditorView::OnShowValueIndicatorsChanged);
+	}
 
 	TSharedRef<SDynamicToolTip> ToolTipWidget =
 		SNew(SDynamicToolTip)
@@ -247,7 +251,6 @@ void SInteractiveCurveEditorView::Tick(const FGeometry& AllottedGeometry, const
 	if (bNeedsRefresh)
 	{
 		bNeedsRefresh = false;
-
 		RefreshRetainer();
 	}
 
@@ -274,7 +277,7 @@ void SInteractiveCurveEditorView::PaintView(const FPaintArgs& Args, const FGeome
 		DrawGridLines(CurveEditorRef, AllottedGeometry, OutDrawElements, BaseLayerId, DrawEffects);
 		DrawBufferedCurves(CurveEditorRef, AllottedGeometry, MyCullingRect, OutDrawElements, BaseLayerId, InWidgetStyle, DrawEffects);
 		DrawCurves(CurveEditorRef, AllottedGeometry, MyCullingRect, OutDrawElements, BaseLayerId, InWidgetStyle, DrawEffects);
-		DrawDottedVisualizationLine(CurveEditorRef, AllottedGeometry, OutDrawElements, BaseLayerId);
+		DrawValueIndicatorLines(CurveEditorRef, AllottedGeometry, OutDrawElements, BaseLayerId);
 	}
 }
 
@@ -935,108 +938,136 @@ void SInteractiveCurveEditorView::DrawBufferedCurves(TSharedRef<FCurveEditor> Cu
 	}
 }
 
-void SInteractiveCurveEditorView::DrawDottedVisualizationLine(
+void SInteractiveCurveEditorView::DrawValueIndicatorLines(
 	TSharedRef<FCurveEditor> InCurveEditor, const FGeometry& InAllottedGeometry, FSlateWindowElementList& OutDrawElements, const int32 InBaseLayerId
 	) const
 {
-	UpdatedKeysWithDottedLines(*InCurveEditor);
+	UpdatedKeysWithValueIndicatorLines(*InCurveEditor);
 	
 	const FCurveEditorScreenSpace ViewSpace = GetViewSpace();
-	for (const TPair<FCurveModelID, TArray<FKeyHandle>>& Pair : CurvesToKeysWithDottedLines)
+	const FCurveModel* CurveModel = ValueIndicatorLineDrawData ? InCurveEditor->FindCurve(ValueIndicatorLineDrawData->HighlightedCurve) : nullptr;
+	if (!CurveModel)
 	{
-		const FCurveModel* CurveModel = InCurveEditor->FindCurve(Pair.Key);
-		if (!CurveModel)
-		{
-			continue;
-		}
-	
-		for (const FKeyHandle& Handle : Pair.Value)
-		{
-			FKeyPosition Position;
-			CurveModel->GetKeyPositions({ Handle }, TArrayView<FKeyPosition>(&Position, 1));
+		return;
+	}
+
+	const auto DrawLine = [&](const FKeyHandle& Handle)
+	{
+		FKeyPosition Position;
+		CurveModel->GetKeyPositions({ Handle }, TArrayView<FKeyPosition>(&Position, 1));
 			
-			const float PosY = ViewSpace.ValueToScreen(Position.OutputValue);
-			FSlateDrawElement::MakeDashedLines(
-				OutDrawElements,
-				InBaseLayerId + CurveViewConstants::ELayerOffset::GridOverlays,
-				InAllottedGeometry.ToPaintGeometry(),
-				{ FVector2f{ 0, PosY }, FVector2f{ InAllottedGeometry.GetLocalSize().X, PosY } }
-				);
-		}
+		const float PosY = ViewSpace.ValueToScreen(Position.OutputValue);
+		FSlateDrawElement::MakeDashedLines(
+			OutDrawElements,
+			InBaseLayerId + CurveViewConstants::ELayerOffset::GridOverlays,
+			InAllottedGeometry.ToPaintGeometry(),
+			{ FVector2f{ 0, PosY }, FVector2f{ InAllottedGeometry.GetLocalSize().X, PosY } },
+			ESlateDrawEffect::None,
+			FLinearColor::White * 0.6f
+			);
+	};
+
+	checkf(ValueIndicatorLineDrawData->MinKey, TEXT("ValueIndicatorLineDrawData should be unset."));
+	DrawLine(ValueIndicatorLineDrawData->MinKey);
+
+	// If only 1 key selected, max is unset.
+	if (ValueIndicatorLineDrawData->MaxKey)
+	{
+		DrawLine(ValueIndicatorLineDrawData->MaxKey);
 	}
 }
 
-void SInteractiveCurveEditorView::UpdatedKeysWithDottedLines(const FCurveEditor& InCurveEditor) const
+void SInteractiveCurveEditorView::UpdatedKeysWithValueIndicatorLines(const FCurveEditor& InCurveEditor) const
 {
-	// We'll draw a horizontal line for each selected point on each selected curve. By default, at most 1 line per curve with selected points.
-	const int32 MaxLinesPerCurve = InCurveEditor.GetSettings()->GetMaxNumHorizontalSelectionLinesPerCurve();
-
+	if (!InCurveEditor.GetSettings()->GetShowValueIndicators())
+	{
+		ValueIndicatorLineDrawData.Reset();
+		return;
+	}
+	
 	const TMap<FCurveModelID, FKeyHandleSet>& Selection = InCurveEditor.Selection.GetAll();
-	for (const TPair<FCurveModelID, FKeyHandleSet>& Pair : Selection)
+	const FCurveModel* Line_CurveModel = ValueIndicatorLineDrawData
+		? InCurveEditor.FindCurve(ValueIndicatorLineDrawData->HighlightedCurve) : nullptr;
+	const FKeyHandleSet* Line_CurveSelectedKeys = ValueIndicatorLineDrawData
+		? Selection.Find(ValueIndicatorLineDrawData->HighlightedCurve) : nullptr;
+	
+	// Clear cached data if it references a curve that is no longer selected or was removed, or the selection contains more than 1 curve.
+	const bool bClearData = !Line_CurveModel || !Line_CurveSelectedKeys || Selection.Num() > 1;
+	if (ValueIndicatorLineDrawData && bClearData)
 	{
-		const FCurveModelID& CurveId = Pair.Key;
-		const FCurveModel* CurveModel = InCurveEditor.FindCurve(CurveId);
-		if (!CurveModel)
-		{
-			continue;
-		}
+		ValueIndicatorLineDrawData.Reset();
+	}
 
-		// If there are multiple points selected, draw dotted lines through the same points every tick until the user deselects them.
-		const FKeyHandleSet& KeyHandleSet = Pair.Value;
-		TArray<FKeyHandle>& KeysWithDottedLines = CurvesToKeysWithDottedLines.FindOrAdd(CurveId);
-		const int32 NumToSet = Algo::RemoveIf(KeysWithDottedLines, [&KeyHandleSet](const FKeyHandle& Handle)
-		{
-			return !KeyHandleSet.Contains(Handle, ECurvePointType::Key);
-		});
-		KeysWithDottedLines.SetNum(
-			// User may have changed the max number last tick
-			FMath::Min(NumToSet, MaxLinesPerCurve)
+	if (ValueIndicatorLineDrawData && !bClearData)
+	{
+		// Update existing data: Min and Max may have changed
+		const bool bIsValid = PickPointsToPlaceValueIndicatorLinesOn(
+			*Line_CurveModel, *Line_CurveSelectedKeys, ValueIndicatorLineDrawData->MinKey, ValueIndicatorLineDrawData->MaxKey
 			);
-
-		// Can add more dotted lines?
-		PickPointsToPlaceDottedLineOn(MaxLinesPerCurve, *CurveModel, KeyHandleSet, KeysWithDottedLines);
+		if (!bIsValid)
+		{
+			ValueIndicatorLineDrawData.Reset();
+		}
 	}
-
-	// Remove dotted lines from curves that are no longer selected.
-	for (auto DottedLineIt = CurvesToKeysWithDottedLines.CreateIterator(); DottedLineIt; ++DottedLineIt)
+	// Only show indicators when exactly 1 curve is selected
+	else if (Selection.Num() == 1)
 	{
-		if (!Selection.Contains(DottedLineIt->Key))
+		// Pick the first applicable curve
+		for (const TPair<FCurveModelID, FKeyHandleSet>& Pair : Selection)
 		{
-			DottedLineIt.RemoveCurrent();
+			const FCurveModelID& CurveId = Pair.Key;
+			const FCurveModel* CurveModel = InCurveEditor.FindCurve(CurveId);
+			if (!CurveModel)
+			{
+				continue;
+			}
+			
+			FKeyHandle Min, Max = FKeyHandle::Invalid();
+			if (PickPointsToPlaceValueIndicatorLinesOn(*CurveModel, Pair.Value, Min, Max))
+			{
+				ValueIndicatorLineDrawData.Emplace(CurveId, Min, Max);
+				return;
+			}
 		}
 	}
 }
 
-void SInteractiveCurveEditorView::PickPointsToPlaceDottedLineOn(
-	const int32 InMaxLinesPerCurve,
+bool SInteractiveCurveEditorView::PickPointsToPlaceValueIndicatorLinesOn(
 	const FCurveModel& InCurveModel,
 	const FKeyHandleSet& InUserSelectedKeys,
-	TArray<FKeyHandle>& OutKeysWithDottedLines
+	FKeyHandle& OutMinKey,
+	FKeyHandle& OutMaxKey
 	) const
 {
-	// Prefer to pick the key with the highest peak.
-	if (InUserSelectedKeys.Num() > 0 && OutKeysWithDottedLines.Num() < InMaxLinesPerCurve)
+	const TConstArrayView<FKeyHandle> Handles = InUserSelectedKeys.AsArray();
+	const bool bContainsNonKeys = Algo::AnyOf(Handles, [&InUserSelectedKeys](const FKeyHandle& Handle)
 	{
-		const TConstArrayView<FKeyHandle> Handles = InUserSelectedKeys.AsArray();
-		const FKeyHandle* PeakPoint = Algo::MaxElementBy(Handles, [&InCurveModel](const FKeyHandle& Handle)
-		{
-			FKeyPosition Position;
-			InCurveModel.GetKeyPositions({ Handle }, TArrayView<FKeyPosition>(&Position, 1));
-			return Position.OutputValue;
-		});
-		OutKeysWithDottedLines.AddUnique(*PeakPoint);
+		return InUserSelectedKeys.PointType(Handle) != ECurvePointType::Key;
+	});
+	// No lines if tangents are selected (or Handles is empty, which should not happen).
+	if (bContainsNonKeys || Handles.Num() == 0)
+	{
+		return false;
 	}
 
-	// Can still add more? Select a 'random' key. 
-	for (int32 Index = 0; Index < InUserSelectedKeys.Num() && OutKeysWithDottedLines.Num() < InMaxLinesPerCurve; Index++)
+	if (Handles.Num() == 1)
 	{
-		const FKeyHandle& Handle = InUserSelectedKeys.AsArray()[Index];
-		const bool bIsPoint = InUserSelectedKeys.PointType(Handle) == ECurvePointType::Key;
-		if (bIsPoint)
-		{
-			OutKeysWithDottedLines.AddUnique(Handle);
-		}
+		OutMinKey = Handles[0];
+		OutMaxKey = FKeyHandle::Invalid();
+		return true;
 	}
+	
+	const auto GetKeyValue = [&InCurveModel](const FKeyHandle& Handle)
+	{
+		FKeyPosition Position;
+		InCurveModel.GetKeyPositions({ Handle }, TArrayView<FKeyPosition>(&Position, 1));
+		return Position.OutputValue;
+	};
+	const FKeyHandle* MinPoint = Algo::MinElementBy(Handles, GetKeyValue);
+	const FKeyHandle* MaxPoint = Algo::MaxElementBy(Handles, GetKeyValue);
+	OutMinKey = *MinPoint;
+	OutMaxKey = OutMinKey == *MaxPoint ? FKeyHandle::Invalid() : *MaxPoint;
+	return true;
 }
 
 bool SInteractiveCurveEditorView::GetPointsWithinWidgetRange(const FSlateRect& WidgetRectangle, TArray<FCurvePointHandle>* OutPoints) const
@@ -2143,4 +2174,10 @@ void SInteractiveCurveEditorView::OnCurveEditorToolChanged(FCurveEditorToolID In
 	}
 }
 
-#undef LOCTEXT_NAMESPACE // "SInteractiveCurveEditorView"
+void SInteractiveCurveEditorView::OnShowValueIndicatorsChanged()
+{
+	// Setting bNeedsRefresh = true is not enough because we may not get ticked.
+	RefreshRetainer();
+}
+
+#undef LOCTEXT_NAMESPACE // "SInteractiveCurveEditorView"
\ No newline at end of file
diff --git a/Engine/Source/Editor/CurveEditor/Public/CurveEditorCommands.h b/Engine/Source/Editor/CurveEditor/Public/CurveEditorCommands.h
--- a/Engine/Source/Editor/CurveEditor/Public/CurveEditorCommands.h
+++ b/Engine/Source/Editor/CurveEditor/Public/CurveEditorCommands.h
@@ -78,6 +78,7 @@ public:
 	TSharedPtr<FUICommandInfo> ToggleShowBufferedCurves;
 	TSharedPtr<FUICommandInfo> ToggleShowCurveEditorCurveToolTips;
 	TSharedPtr<FUICommandInfo> ToggleShowBars;
+	TSharedPtr<FUICommandInfo> ToggleShowValueIndicatorLines;
 
 	TSharedPtr<FUICommandInfo> SetRandomCurveColorsForSelected;
 	TSharedPtr<FUICommandInfo> SetCurveColorsForSelected;
diff --git a/Engine/Source/Editor/CurveEditor/Public/CurveEditorSettings.h b/Engine/Source/Editor/CurveEditor/Public/CurveEditorSettings.h
--- a/Engine/Source/Editor/CurveEditor/Public/CurveEditorSettings.h
+++ b/Engine/Source/Editor/CurveEditor/Public/CurveEditorSettings.h
@@ -94,6 +94,7 @@ public:
 
 	DECLARE_MULTICAST_DELEGATE(FOnCustomColorsChanged);
 	DECLARE_MULTICAST_DELEGATE(FOnAxisSnappingChanged);
+	DECLARE_MULTICAST_DELEGATE(FOnShowValueIndicatorsChanged);
 
 	UCurveEditorSettings();
 
@@ -161,9 +162,10 @@ public:
 	
 	/** Gets the multicast delegate which is run whenever custom colors have changed. */
 	FOnCustomColorsChanged& GetOnCustomColorsChanged() { return OnCustomColorsChangedEvent; }
-
 	/** Gets the multicast delegate which is run whenever axis snapping has changed. */
 	FOnAxisSnappingChanged& GetOnAxisSnappingChanged() { return OnAxisSnappingChangedEvent; }
+	/** Gets the multicast delegate which is run whenever showing the value indicator lines has changed. */
+	FOnShowValueIndicatorsChanged& OnShowValueIndicatorsChanged() { return OnShowValueIndicatorsChangedEvent; }
 
 	/** Get custom color for space name. Parent and World are reserved names and will be used instead of the specified control name. */
 	TOptional<FLinearColor> GetSpaceSwitchColor(const FString& InControlName) const;
@@ -185,10 +187,10 @@ public:
 	/** Sets how sensitive the selection marquee should be when selecting points. */
 	void SetMarqueePointSensitivity(float InMarqueePointSensitivity);
 
-	/** Gets how sensitive the selection marquee should be when selecting points. */
-	float GetMaxNumHorizontalSelectionLinesPerCurve() const { return MaxNumHorizontalSelectionLinesPerCurve; }
-	/** Sets how sensitive the selection marquee should be when selecting points. */
-	void SetMaxNumHorizontalSelectionLinesPerCurve(float InValue);
+	/** @return Whether to draw a value indicator line for the minimum and maximum key in the selected key range. */
+	bool GetShowValueIndicators() const { return bShowValueIndicators; }
+	/** Sets whether to draw a value indicator line for the minimum and maximum key in the selected key range. */
+	void SetShowValueIndicators(bool bValue);
 
 protected:
 
@@ -254,10 +256,12 @@ protected:
 	UPROPERTY(config, EditAnywhere, Category = "Curve Editor", meta = (ClampMin = 0.0, ClampMax = 1.0))
 	float MarqueePointSensitivity = 0.45f;
 
-	UPROPERTY(config, EditAnywhere, Category = "Curve Editor", meta = (ClampMin = 0))
-	int32 MaxNumHorizontalSelectionLinesPerCurve = 1;
+	/** When you select a single curve, whether to draw a dotted line for the minimum and maximum key in the selected key range. */
+	UPROPERTY(config, EditAnywhere, Category = "Curve Editor")
+	bool bShowValueIndicators = true;
 
 private:
 	FOnCustomColorsChanged OnCustomColorsChangedEvent;
 	FOnAxisSnappingChanged OnAxisSnappingChangedEvent;
+	FOnShowValueIndicatorsChanged OnShowValueIndicatorsChangedEvent;
 };
diff --git a/Engine/Source/Editor/CurveEditor/Public/Views/SInteractiveCurveEditorView.h b/Engine/Source/Editor/CurveEditor/Public/Views/SInteractiveCurveEditorView.h
--- a/Engine/Source/Editor/CurveEditor/Public/Views/SInteractiveCurveEditorView.h
+++ b/Engine/Source/Editor/CurveEditor/Public/Views/SInteractiveCurveEditorView.h
@@ -139,17 +139,18 @@ protected:
 	void DrawGridLines(TSharedRef<FCurveEditor> CurveEditor, const FGeometry& AllottedGeometry, FSlateWindowElementList& OutDrawElements, int32 BaseLayerId, ESlateDrawEffect DrawEffects) const;
 	void DrawCurves(TSharedRef<FCurveEditor> CurveEditor, const FGeometry& AllottedGeometry, const FSlateRect& MyCullingRect, FSlateWindowElementList& OutDrawElements, int32 BaseLayerId, const FWidgetStyle& InWidgetStyle, ESlateDrawEffect DrawEffects) const;
 	void DrawBufferedCurves(TSharedRef<FCurveEditor> CurveEditor, const FGeometry& AllottedGeometry, const FSlateRect& MyCullingRect, FSlateWindowElementList& OutDrawElements, int32 BaseLayerId, const FWidgetStyle& InWidgetStyle, ESlateDrawEffect DrawEffects) const;
-	void DrawDottedVisualizationLine(TSharedRef<FCurveEditor> InCurveEditor, const FGeometry& InAllottedGeometry, FSlateWindowElementList& OutDrawElements, int32 InBaseLayerId) const;
-
-	/** Updates the keys that will be drawn with dotted lines. */
-	void UpdatedKeysWithDottedLines(const FCurveEditor& InCurveEditor) const;
-	/** Picks which keys will have dotted lines drawn for them. By default, each curve picks the key in the user selection with the highest peak. */
-	virtual void PickPointsToPlaceDottedLineOn(const int32 InMaxLinesPerCurve, const FCurveModel& InCurveModel, const FKeyHandleSet& InUserSelectedKeys, TArray<FKeyHandle>& OutKeysWithDottedLines) const;
+	void DrawValueIndicatorLines(TSharedRef<FCurveEditor> InCurveEditor, const FGeometry& InAllottedGeometry, FSlateWindowElementList& OutDrawElements, int32 InBaseLayerId) const;
+	
+	/** Updates the keys that will be drawn with value indicator lines. */
+	void UpdatedKeysWithValueIndicatorLines(const FCurveEditor& InCurveEditor) const;
+	/** Picks which keys will have value indicator lines drawn for them. Returns true if the selection. */
+	bool PickPointsToPlaceValueIndicatorLinesOn(const FCurveModel& InCurveModel, const FKeyHandleSet& InUserSelectedKeys, FKeyHandle& OutMinKey, FKeyHandle& OutMaxKey) const;
 	
 	FSlateColor GetCurveCaptionColor() const;
 	FText GetCurveCaption() const;
 
 private:
+	
 	void HandleDirectKeySelectionByMouse(TSharedPtr<FCurveEditor> CurveEditor, const FPointerEvent& MouseEvent, TOptional<FCurvePointHandle> MouseDownPoint);
 
 	void CreateContextMenu(const FGeometry& MyGeometry, const FPointerEvent& MouseEvent);
@@ -182,6 +183,7 @@ private:
 	void PasteKeys(TSet<FCurveModelID> ToCurves);
 
 	void OnCurveEditorToolChanged(FCurveEditorToolID InToolId);
+	void OnShowValueIndicatorsChanged();
 
 	/**
 	 * Rebind contextual command mappings that rely on the mouse position
@@ -204,15 +206,6 @@ protected:
 	/** Background tint for this widget */
 	FLinearColor BackgroundTint;
 
-	/**
-	 * We draw a horizontal dotted line on the keys that a user selects.
-	 * This makes it easier for the user to see a selected key's value in relation to other keys.
-	 * The  per-curve upper limit of dotted lines that are drawn is configured in UCurveEditorSettings::MaxNumHorizontalSelectionLinesPerCurve.
-	 *
-	 * This is called in OnPaint, which is const. Mutable is a workaround.
-	 */
-	mutable TMap<FCurveModelID, TArray<FKeyHandle>> CurvesToKeysWithDottedLines;
-
 private:
 
 	/** (Optional) the current drag operation */
@@ -245,4 +238,32 @@ private:
 	mutable FSlateColor CachedCurveCaptionColor;
 
 	mutable bool bNeedsRefresh = false;
+	
+	struct FValueIndicatorLineData
+	{
+		/** The curve for which the value indicator lines are being drawn. */
+		const FCurveModelID HighlightedCurve;
+
+		/** Key handle to the min item through which a value indicator line is being drawn. */
+		FKeyHandle MinKey = FKeyHandle::Invalid();
+		/** Key handle to the max item through which a value indicator line is being drawn. If the user only has 1 key selected, this is FKeyHandle::Invalid(). */
+		FKeyHandle MaxKey = FKeyHandle::Invalid();
+
+		explicit FValueIndicatorLineData(const FCurveModelID& HighlightedCurve, const FKeyHandle& MinKey, const FKeyHandle& MaxKey)
+			: HighlightedCurve(HighlightedCurve) , MinKey(MinKey), MaxKey(MaxKey)
+		{}
+	};
+	/**
+	 * We draw up to 2 horizontal value indicator line on the keys that a user selects.
+	 * This makes it easier for the user to see a selected key's value in relation to other keys.
+	 *
+	 * Requirements:
+	 * - Only draw value indicator line if 1 curve is selected. More curves selected = no line.
+	 * - The value indicator lines are drawn through the min and max value keys in the selected range.
+	 * - If only 1 key is selected, only draw 1 line (obviously). 
+	 * The per-curve upper limit of value indicator lines that are drawn is configured in UCurveEditorSettings::MaxNumHorizontalSelectionLinesPerCurve.
+	 *
+	 * This is called in OnPaint, which is const. Mutable is a workaround.
+	 */
+	mutable TOptional<FValueIndicatorLineData> ValueIndicatorLineDrawData;
 };
\ No newline at end of file
diff --git a/Engine/Source/Runtime/Engine/Classes/Curves/KeyHandle.h b/Engine/Source/Runtime/Engine/Classes/Curves/KeyHandle.h
--- a/Engine/Source/Runtime/Engine/Classes/Curves/KeyHandle.h
+++ b/Engine/Source/Runtime/Engine/Classes/Curves/KeyHandle.h
@@ -17,6 +17,10 @@ struct FKeyHandle
 
 	static ENGINE_API FKeyHandle Invalid();
 
+	ENGINE_API bool IsValid() const;
+
+	operator bool() const { return IsValid(); }
+
 	bool operator ==(const FKeyHandle& Other) const
 	{
 		return Index == Other.Index;
diff --git a/Engine/Source/Runtime/Engine/Private/Curves/KeyHandle.cpp b/Engine/Source/Runtime/Engine/Private/Curves/KeyHandle.cpp
--- a/Engine/Source/Runtime/Engine/Private/Curves/KeyHandle.cpp
+++ b/Engine/Source/Runtime/Engine/Private/Curves/KeyHandle.cpp
@@ -39,6 +39,11 @@ FKeyHandle FKeyHandle::Invalid()
 	return FKeyHandle(0);
 }
 
+bool FKeyHandle::IsValid() const
+{
+	return Index != 0;
+}
+
 void FKeyHandleMap::Initialize(TArrayView<const FKeyHandle> InKeyHandles)
 {
 	Empty(InKeyHandles.Num());
