Curve Editor: Improve drag operation for moving keys performance.
- Mouse move events are now accumulated
- Once the engine is done pumping mouse movements, we process the final mouse location
- Note: OnMouse move may be called multiple times within the same frame!

diff --git a/Engine/Source/Editor/CurveEditor/Private/DragOperations/CurveEditorDragOperation_MoveKeys.cpp b/Engine/Source/Editor/CurveEditor/Private/DragOperations/CurveEditorDragOperation_MoveKeys.cpp
--- a/Engine/Source/Editor/CurveEditor/Private/DragOperations/CurveEditorDragOperation_MoveKeys.cpp
+++ b/Engine/Source/Editor/CurveEditor/Private/DragOperations/CurveEditorDragOperation_MoveKeys.cpp
@@ -66,84 +66,20 @@
 
 void FCurveEditorDragOperation_MoveKeys::OnDrag(FVector2D InitialPosition, FVector2D CurrentPosition, const FPointerEvent& MouseEvent)
 {
-	TArray<FKeyPosition> NewKeyPositionScratch;
-	FVector2D MousePosition = CurveEditor->GetAxisSnap().GetSnappedPosition(InitialPosition, LastMousePosition, CurrentPosition, MouseEvent, SnappingState);
+	// OnDrag may be called multiple times in the same frame. We'll accumulate the end position and process it in OnFinishedPointerInput.
+	GetOrAddAccumulatedMouseMovement(InitialPosition).EndMousePosition = CurveEditor->GetAxisSnap().GetSnappedPosition(
+		InitialPosition, LastMousePosition, CurrentPosition, MouseEvent, SnappingState
+		);
 	LastMousePosition = CurrentPosition;
+}
 
-	for (FKeyData& KeyData : KeysByCurve)
+void FCurveEditorDragOperation_MoveKeys::OnFinishedPointerInput()
+{
+	// OnFinishedPointerInput is called once the engine is done pumping messages.
+	if (AccumulatedMouseMovement)
 	{
-		const SCurveEditorView* View = CurveEditor->FindFirstInteractiveView(KeyData.CurveID);
-		if (!View)
-		{
-			continue;
-		}
-
-		FCurveModel* Curve = CurveEditor->FindCurve(KeyData.CurveID);
-		if (!ensureAlways(Curve))
-		{
-			continue;
-		}
-
-		FCurveEditorScreenSpace CurveSpace = View->GetCurveSpace(KeyData.CurveID);
-
-		double DeltaInput = (MousePosition.X - InitialPosition.X) / CurveSpace.PixelsPerInput();
-		double DeltaOutput = -(MousePosition.Y - InitialPosition.Y) / CurveSpace.PixelsPerOutput();
-
-		NewKeyPositionScratch.Reset();
-		NewKeyPositionScratch.Reserve(KeyData.StartKeyPositions.Num());
-
-		FCurveSnapMetrics SnapMetrics = CurveEditor->GetCurveSnapMetrics(KeyData.CurveID);
-
-		// Transform by the inverse curve transform to put the start positions in view space
-		FTransform2d CurveTransform               = Curve->GetCurveTransform();
-		FTransform2d InverseInitialCurveTransform = KeyData.InitialDragTransform.Inverse();
-
-		if (CardinalPoint.IsSet())
-		{
-			for (int KeyIndex = 0; KeyIndex < KeyData.StartKeyPositions.Num(); ++KeyIndex)
-			{
-				FKeyHandle KeyHandle = KeyData.Handles[KeyIndex];
-				if (CardinalPoint->KeyHandle == KeyHandle)
-				{
-					FKeyPosition StartPosition = KeyData.StartKeyPositions[KeyIndex].Transform(InverseInitialCurveTransform);
-
-					if (View->IsTimeSnapEnabled())
-					{
-						DeltaInput = SnapMetrics.SnapInputSeconds(StartPosition.InputValue + DeltaInput) - StartPosition.InputValue;
-					}
-
-					// If view is not absolute, snap based on the key that was grabbed, not all keys individually.
-					if (View->IsValueSnapEnabled() && View->ViewTypeID != ECurveEditorViewID::Absolute)
-					{
-						DeltaOutput = SnapMetrics.SnapOutput(StartPosition.OutputValue + DeltaOutput) - StartPosition.OutputValue;
-					}
-
-					break;
-				}
-			}
-		}
-		for (FKeyPosition StartPosition : KeyData.StartKeyPositions)
-		{
-			StartPosition = StartPosition.Transform(InverseInitialCurveTransform);
-
-			StartPosition.InputValue += DeltaInput;
-			StartPosition.OutputValue += DeltaOutput;
-
-			StartPosition.InputValue = View->IsTimeSnapEnabled() ? SnapMetrics.SnapInputSeconds(StartPosition.InputValue) : StartPosition.InputValue;
-
-			// Snap value keys individually if view mode is absolute.
-			if (View->ViewTypeID == ECurveEditorViewID::Absolute)
-			{
-				StartPosition.OutputValue = View->IsValueSnapEnabled() ? SnapMetrics.SnapOutput(StartPosition.OutputValue) : StartPosition.OutputValue;
-			}
-
-			NewKeyPositionScratch.Add(StartPosition.Transform(CurveTransform));
-		}
-
-		Curve->SetKeyPositions(KeyData.Handles, NewKeyPositionScratch, EPropertyChangeType::Interactive);
-
-		// Make sure the last dragged key positions are up to date
-		Curve->GetKeyPositions(KeyData.Handles, KeyData.LastDraggedKeyPositions);
+		UpdateFromDrag(AccumulatedMouseMovement->InitialPosition, AccumulatedMouseMovement->EndMousePosition);
+		AccumulatedMouseMovement.Reset();
 	}
 }
 
@@ -230,3 +166,94 @@
 
 	CurveEditor->SuppressBoundTransformUpdates(false);
 }
+
+FCurveEditorDragOperation_MoveKeys::FAccumulatedMouseMovement& FCurveEditorDragOperation_MoveKeys::GetOrAddAccumulatedMouseMovement(
+	const FVector2D& InitialPosition
+	)
+{
+	if (!AccumulatedMouseMovement)
+	{
+		AccumulatedMouseMovement.Emplace(InitialPosition);
+	}
+	return *AccumulatedMouseMovement;
+}
+
+void FCurveEditorDragOperation_MoveKeys::UpdateFromDrag(const FVector2D& InitialPosition, const FVector2D& MousePosition)
+{
+	TArray<FKeyPosition> NewKeyPositionScratch;
+	for (FKeyData& KeyData : KeysByCurve)
+	{
+		const SCurveEditorView* View = CurveEditor->FindFirstInteractiveView(KeyData.CurveID);
+		if (!View)
+		{
+			continue;
+		}
+
+		FCurveModel* Curve = CurveEditor->FindCurve(KeyData.CurveID);
+		if (!ensureAlways(Curve))
+		{
+			continue;
+		}
+
+		FCurveEditorScreenSpace CurveSpace = View->GetCurveSpace(KeyData.CurveID);
+
+		double DeltaInput = (MousePosition.X - InitialPosition.X) / CurveSpace.PixelsPerInput();
+		double DeltaOutput = -(MousePosition.Y - InitialPosition.Y) / CurveSpace.PixelsPerOutput();
+
+		NewKeyPositionScratch.Reset();
+		NewKeyPositionScratch.Reserve(KeyData.StartKeyPositions.Num());
+
+		FCurveSnapMetrics SnapMetrics = CurveEditor->GetCurveSnapMetrics(KeyData.CurveID);
+
+		// Transform by the inverse curve transform to put the start positions in view space
+		FTransform2d CurveTransform               = Curve->GetCurveTransform();
+		FTransform2d InverseInitialCurveTransform = KeyData.InitialDragTransform.Inverse();
+
+		if (CardinalPoint.IsSet())
+		{
+			for (int KeyIndex = 0; KeyIndex < KeyData.StartKeyPositions.Num(); ++KeyIndex)
+			{
+				FKeyHandle KeyHandle = KeyData.Handles[KeyIndex];
+				if (CardinalPoint->KeyHandle == KeyHandle)
+				{
+					FKeyPosition StartPosition = KeyData.StartKeyPositions[KeyIndex].Transform(InverseInitialCurveTransform);
+
+					if (View->IsTimeSnapEnabled())
+					{
+						DeltaInput = SnapMetrics.SnapInputSeconds(StartPosition.InputValue + DeltaInput) - StartPosition.InputValue;
+					}
+
+					// If view is not absolute, snap based on the key that was grabbed, not all keys individually.
+					if (View->IsValueSnapEnabled() && View->ViewTypeID != ECurveEditorViewID::Absolute)
+					{
+						DeltaOutput = SnapMetrics.SnapOutput(StartPosition.OutputValue + DeltaOutput) - StartPosition.OutputValue;
+					}
+
+					break;
+				}
+			}
+		}
+		for (FKeyPosition StartPosition : KeyData.StartKeyPositions)
+		{
+			StartPosition = StartPosition.Transform(InverseInitialCurveTransform);
+
+			StartPosition.InputValue += DeltaInput;
+			StartPosition.OutputValue += DeltaOutput;
+
+			StartPosition.InputValue = View->IsTimeSnapEnabled() ? SnapMetrics.SnapInputSeconds(StartPosition.InputValue) : StartPosition.InputValue;
+
+			// Snap value keys individually if view mode is absolute.
+			if (View->ViewTypeID == ECurveEditorViewID::Absolute)
+			{
+				StartPosition.OutputValue = View->IsValueSnapEnabled() ? SnapMetrics.SnapOutput(StartPosition.OutputValue) : StartPosition.OutputValue;
+			}
+
+			NewKeyPositionScratch.Add(StartPosition.Transform(CurveTransform));
+		}
+
+		Curve->SetKeyPositions(KeyData.Handles, NewKeyPositionScratch, EPropertyChangeType::Interactive);
+
+		// Make sure the last dragged key positions are up to date
+		Curve->GetKeyPositions(KeyData.Handles, KeyData.LastDraggedKeyPositions);
+	}
+}
diff --git a/Engine/Source/Editor/CurveEditor/Private/DragOperations/CurveEditorDragOperation_MoveKeys.h b/Engine/Source/Editor/CurveEditor/Private/DragOperations/CurveEditorDragOperation_MoveKeys.h
--- a/Engine/Source/Editor/CurveEditor/Private/DragOperations/CurveEditorDragOperation_MoveKeys.h
+++ b/Engine/Source/Editor/CurveEditor/Private/DragOperations/CurveEditorDragOperation_MoveKeys.h
@@ -17,11 +17,15 @@ struct FPointerEvent;
 class FCurveEditorDragOperation_MoveKeys : public ICurveEditorKeyDragOperation
 {
 public:
+
+	//~ Begin ICurveEditorDragOperation Interface
 	virtual void OnInitialize(FCurveEditor* CurveEditor, const TOptional<FCurvePointHandle>& CardinalPoint) override;
 	virtual void OnBeginDrag(FVector2D InitialPosition, FVector2D CurrentPosition, const FPointerEvent& MouseEvent) override;
 	virtual void OnDrag(FVector2D InitialPosition, FVector2D CurrentPosition, const FPointerEvent& MouseEvent) override;
+	virtual void OnFinishedPointerInput() override;
 	virtual void OnEndDrag(FVector2D InitialPosition, FVector2D CurrentPosition, const FPointerEvent& MouseEvent) override;
 	virtual void OnCancelDrag() override;
+	//~ End ICurveEditorDragOperation Interface
 
 private:
 
@@ -54,4 +58,19 @@ private:
 	TArray<FKeyData> KeysByCurve;
 
 	FCurveEditorAxisSnap::FSnapState SnappingState;
+
+	struct FAccumulatedMouseMovement
+	{
+		const FVector2D InitialPosition;
+		FVector2D EndMousePosition;
+
+		explicit FAccumulatedMouseMovement(const FVector2D& InitialPosition)
+			: InitialPosition(InitialPosition)
+		{}
+	};
+	TOptional<FAccumulatedMouseMovement> AccumulatedMouseMovement;
+
+	FAccumulatedMouseMovement& GetOrAddAccumulatedMouseMovement(const FVector2D& InitialPosition);
+
+	void UpdateFromDrag(const FVector2D& InitialPosition, const FVector2D& MousePosition);
 };
\ No newline at end of file
diff --git a/Engine/Source/Editor/CurveEditor/Private/DragOperations/ICurveEditorDragOperation.cpp b/Engine/Source/Editor/CurveEditor/Private/DragOperations/ICurveEditorDragOperation.cpp
--- a/Engine/Source/Editor/CurveEditor/Private/DragOperations/ICurveEditorDragOperation.cpp
+++ b/Engine/Source/Editor/CurveEditor/Private/DragOperations/ICurveEditorDragOperation.cpp
@@ -17,6 +17,11 @@ void ICurveEditorDragOperation::Drag(FVector2D InitialPosition, FVector2D Curren
 	OnDrag(InitialPosition, CurrentPosition, MouseEvent);
 }
 
+void ICurveEditorDragOperation::FinishedPointerInput()
+{
+	OnFinishedPointerInput();
+}
+
 FReply ICurveEditorDragOperation::MouseWheel(FVector2D InitialPosition, FVector2D CurrentPosition, const FPointerEvent& MouseEvent)
 {
 	return OnMouseWheel(InitialPosition, CurrentPosition, MouseEvent);
diff --git a/Engine/Source/Editor/CurveEditor/Private/SCurveEditorViewContainer.cpp b/Engine/Source/Editor/CurveEditor/Private/SCurveEditorViewContainer.cpp
--- a/Engine/Source/Editor/CurveEditor/Private/SCurveEditorViewContainer.cpp
+++ b/Engine/Source/Editor/CurveEditor/Private/SCurveEditorViewContainer.cpp
@@ -287,7 +287,8 @@ FReply SCurveEditorViewContainer::OnMouseMove(const FGeometry& MyGeometry, const
 {
 	FVector2D MousePixel = MyGeometry.AbsoluteToLocal(MouseEvent.GetScreenSpacePosition());
 	FReply ToolReply = FReply::Unhandled();
-
+	bHadMouseMovesThisTick = true;
+	
 	if (CurveEditor->GetCurrentTool())
 	{
 		ToolReply = CurveEditor->GetCurrentTool()->OnMouseMove(AsShared(), MyGeometry, MouseEvent);
@@ -312,6 +313,26 @@ FReply SCurveEditorViewContainer::OnMouseMove(const FGeometry& MyGeometry, const
 	return ToolReply;
 }
 
+void SCurveEditorViewContainer::OnFinishedPointerInput()
+{
+	if (!bHadMouseMovesThisTick)
+	{
+		return;
+	}
+	// Some tools & operations defer processing Drag calls for performance reasons. Give them a chance to process the accumulated input.
+	
+	if (CurveEditor->GetCurrentTool())
+	{
+		CurveEditor->GetCurrentTool()->OnFinishedPointerInput();
+	}
+	if (DragOperation && DragOperation->IsDragging())
+	{
+		DragOperation->DragImpl->FinishedPointerInput();
+	}
+	
+	bHadMouseMovesThisTick = false;
+}
+
 FReply SCurveEditorViewContainer::OnMouseButtonUp(const FGeometry& MyGeometry, const FPointerEvent& MouseEvent)
 {
 	FVector2D MousePosition = MyGeometry.AbsoluteToLocal(MouseEvent.GetScreenSpacePosition());
diff --git a/Engine/Source/Editor/CurveEditor/Private/SCurveEditorViewContainer.h b/Engine/Source/Editor/CurveEditor/Private/SCurveEditorViewContainer.h
--- a/Engine/Source/Editor/CurveEditor/Private/SCurveEditorViewContainer.h
+++ b/Engine/Source/Editor/CurveEditor/Private/SCurveEditorViewContainer.h
@@ -70,6 +70,7 @@ private:
 	virtual FReply OnMouseButtonDown(const FGeometry& MyGeometry, const FPointerEvent& MouseEvent) override;
 	virtual FReply OnMouseButtonUp(const FGeometry& MyGeometry, const FPointerEvent& MouseEvent) override;
 	virtual FReply OnMouseMove(const FGeometry& MyGeometry, const FPointerEvent& MouseEvent) override;
+	virtual void OnFinishedPointerInput() override;
 	virtual FReply OnMouseButtonDoubleClick(const FGeometry& InMyGeometry, const FPointerEvent& InMouseEvent) override;
 	virtual void OnFocusLost(const FFocusEvent& InFocusEvent) override;
 	virtual FCursorReply OnCursorQuery(const FGeometry& MyGeometry, const FPointerEvent& CursorEvent) const override;
@@ -98,6 +99,8 @@ private:
 
 	/** (Optional) the current drag operation */
 	TOptional<FCurveEditorDelayedDrag> DragOperation;
+	/** Whether DragOperation needs to have OnFinishedPointerInput called this tick. */
+	bool bHadMouseMovesThisTick = false;
 
 	/** Array of views whose cache may need to be updated and may need their height updating on tick. */
 	TArray<TSharedPtr<SCurveEditorView>> Views;
diff --git a/Engine/Source/Editor/CurveEditor/Private/Views/SInteractiveCurveEditorView.cpp b/Engine/Source/Editor/CurveEditor/Private/Views/SInteractiveCurveEditorView.cpp
--- a/Engine/Source/Editor/CurveEditor/Private/Views/SInteractiveCurveEditorView.cpp
+++ b/Engine/Source/Editor/CurveEditor/Private/Views/SInteractiveCurveEditorView.cpp
@@ -1319,6 +1319,7 @@ FReply SInteractiveCurveEditorView::OnMouseMove(const FGeometry& MyGeometry, con
 		}
 		else if (DragOperation->IsDragging())
 		{
+			bHadMouseMovesThisTick = true;
 			DragOperation->DragImpl->Drag(InitialPosition, MousePixel, MouseEvent);
 		}
 		return FReply::Handled();
@@ -1328,6 +1329,16 @@ FReply SInteractiveCurveEditorView::OnMouseMove(const FGeometry& MyGeometry, con
 	return FReply::Unhandled();
 }
 
+void SInteractiveCurveEditorView::OnFinishedPointerInput()
+{
+	// Some operations defer processing Drag calls for performance reasons. Give them a chance to process the accumulated input.
+	if (bHadMouseMovesThisTick && DragOperation && DragOperation->IsDragging())
+	{
+		DragOperation->DragImpl->FinishedPointerInput();
+	}
+	bHadMouseMovesThisTick = false;
+}
+
 FReply SInteractiveCurveEditorView::OnMouseWheel(const FGeometry& MyGeometry, const FPointerEvent& MouseEvent)
 {
 	TSharedPtr<FCurveEditor> CurveEditor = WeakCurveEditor.Pin();
diff --git a/Engine/Source/Editor/CurveEditor/Public/ICurveEditorDragOperation.h b/Engine/Source/Editor/CurveEditor/Public/ICurveEditorDragOperation.h
--- a/Engine/Source/Editor/CurveEditor/Public/ICurveEditorDragOperation.h
+++ b/Engine/Source/Editor/CurveEditor/Public/ICurveEditorDragOperation.h
@@ -41,6 +41,9 @@ public:
 	 */
 	void Drag(FVector2D InitialPosition, FVector2D CurrentPosition, const FPointerEvent& MouseEvent);
 
+	/** Process any mouse input that this drag operation may have accumulated. */
+	void FinishedPointerInput();
+
 	/**
 	 * Potentially Evaluate a MouseWheel event which occcured during this drag operation
 	 */
@@ -71,6 +74,10 @@ protected:
 	virtual void OnDrag(FVector2D InitialPosition, FVector2D CurrentPosition, const FPointerEvent& MouseEvent)
 	{}
 
+	/** Implementation method for derived types to process any accumulated drags. */
+	virtual void OnFinishedPointerInput()
+	{}
+
 	/** Implementation method for derived types to evaluate a mousewheel event */
 	virtual FReply OnMouseWheel(FVector2D InitialPosition, FVector2D CurrentPosition, const FPointerEvent& MouseEvent)
 	{
diff --git a/Engine/Source/Editor/CurveEditor/Public/ICurveEditorToolExtension.h b/Engine/Source/Editor/CurveEditor/Public/ICurveEditorToolExtension.h
--- a/Engine/Source/Editor/CurveEditor/Public/ICurveEditorToolExtension.h
+++ b/Engine/Source/Editor/CurveEditor/Public/ICurveEditorToolExtension.h
@@ -95,6 +95,7 @@ public:
 	virtual FReply OnMouseButtonUp(TSharedRef<SWidget> OwningWidget, const FGeometry& MyGeometry, const FPointerEvent& MouseEvent) { return FReply::Unhandled(); }
 	virtual FReply OnMouseMove(TSharedRef<SWidget> OwningWidget, const FGeometry& MyGeometry, const FPointerEvent& MouseEvent) { return FReply::Unhandled(); }
 	virtual FReply OnMouseButtonDoubleClick(TSharedRef<SWidget> OwningWidget, const FGeometry& InMyGeometry, const FPointerEvent& InMouseEvent) { return FReply::Unhandled(); }
+	virtual void OnFinishedPointerInput() {}
 	virtual void OnFocusLost(const FFocusEvent& InFocusEvent) {}
 	virtual TSharedPtr<FStructOnScope> GetToolOptions() const { return nullptr; }
 	virtual void OnToolOptionsUpdated(const FPropertyChangedEvent& PropertyChangedEvent) {}
diff --git a/Engine/Source/Editor/CurveEditor/Public/Views/SInteractiveCurveEditorView.h b/Engine/Source/Editor/CurveEditor/Public/Views/SInteractiveCurveEditorView.h
--- a/Engine/Source/Editor/CurveEditor/Public/Views/SInteractiveCurveEditorView.h
+++ b/Engine/Source/Editor/CurveEditor/Public/Views/SInteractiveCurveEditorView.h
@@ -128,6 +128,7 @@ protected:
 	virtual void OnMouseEnter(const FGeometry& MyGeometry, const FPointerEvent& MouseEvent) override;
 	virtual void OnMouseLeave(const FPointerEvent& MouseEvent) override;
 	virtual FReply OnMouseMove(const FGeometry& MyGeometry, const FPointerEvent& MouseEvent) override;
+	virtual void OnFinishedPointerInput() override;
 	virtual FReply OnMouseWheel(const FGeometry& MyGeometry, const FPointerEvent& MouseEvent) override;
 	virtual FReply OnMouseButtonUp(const FGeometry& MyGeometry, const FPointerEvent& MouseEvent) override;
 	virtual FReply OnMouseButtonDown(const FGeometry& MyGeometry, const FPointerEvent& MouseEvent) override;
@@ -210,6 +211,8 @@ private:
 
 	/** (Optional) the current drag operation */
 	TOptional<FCurveEditorDelayedDrag> DragOperation;
+	/** Whether DragOperation needs to have OnFinishedPointerInput called this tick. */
+	bool bHadMouseMovesThisTick = false;
 
 	struct FCachedToolTipData
 	{
