Curve Editor: Change how weighted tangents can be edited

* Shift-LMB drag now locks the tangent
* Ctrl-Shift-LMB drag now locks the weight
* Shift-LMB no longer snaps to 45 degrees

diff --git a/Engine/Source/Editor/CurveEditor/Private/DragOperations/CurveEditorDragOperation_Tangent.cpp b/Engine/Source/Editor/CurveEditor/Private/DragOperations/CurveEditorDragOperation_Tangent.cpp
--- a/Engine/Source/Editor/CurveEditor/Private/DragOperations/CurveEditorDragOperation_Tangent.cpp
+++ b/Engine/Source/Editor/CurveEditor/Private/DragOperations/CurveEditorDragOperation_Tangent.cpp
@@ -27,8 +27,192 @@
 
 namespace CurveEditorDragOperation
 {
-	/** If they're within this many pixels of crossing over the straight up/down boundary, we clamp the handles. */
-	constexpr float TangentCrossoverThresholdPx = 1.f;
+	FTangentSolver::FTangentSolver(ECurvePointType InPointType, FKeyAttributes InInitialKeyAttributes)
+		: InitialKeyAttributes(InInitialKeyAttributes)
+		, PointType(InPointType)
+	{}
+
+	void FTangentSolver::Process(
+		const FPointerEvent& MouseEvent, 
+		const FCurveEditorScreenSpace& CurveSpace, 
+		const FVector2D& PixelDelta)
+	{
+		const float DisplayRatio = (CurveSpace.PixelsPerOutput() / CurveSpace.PixelsPerInput());
+
+		if (InitialKeyAttributes.HasArriveTangent() && EnumHasAnyFlags(PointType, ECurvePointType::ArriveTangent))
+		{
+			const bool bArriveTangent = true;
+			const float ArriveTangent = InitialKeyAttributes.GetArriveTangent();
+
+			const FSolverParams SolverParams(
+				MouseEvent,
+				CurveSpace,
+				PixelDelta,
+				DisplayRatio,
+				bArriveTangent,
+				ArriveTangent);
+
+			if (InitialKeyAttributes.HasTangentWeightMode() && InitialKeyAttributes.HasArriveTangentWeight() &&
+				(InitialKeyAttributes.GetTangentWeightMode() == RCTWM_WeightedBoth || InitialKeyAttributes.GetTangentWeightMode() == RCTWM_WeightedArrive))
+			{
+				SolveTangentMutableWeight(SolverParams, InitialKeyAttributes.GetArriveTangentWeight());
+			}
+			else
+			{
+				SolveTangentConstantWeight(SolverParams);
+			}
+		}
+
+		if (InitialKeyAttributes.HasLeaveTangent() && EnumHasAnyFlags(PointType, ECurvePointType::LeaveTangent))
+		{
+			const bool bArriveTangent = false;
+			const float LeaveTangent = InitialKeyAttributes.GetLeaveTangent();
+
+			const FSolverParams SolverParams(
+				MouseEvent,
+				CurveSpace,
+				PixelDelta,
+				DisplayRatio,
+				bArriveTangent,
+				LeaveTangent);
+
+			if (InitialKeyAttributes.HasTangentWeightMode() && InitialKeyAttributes.HasLeaveTangentWeight() &&
+				(InitialKeyAttributes.GetTangentWeightMode() == RCTWM_WeightedBoth || InitialKeyAttributes.GetTangentWeightMode() == RCTWM_WeightedLeave))
+			{
+				SolveTangentMutableWeight(SolverParams, InitialKeyAttributes.GetLeaveTangentWeight());
+			}
+			else
+			{
+				SolveTangentConstantWeight(SolverParams);
+			}
+		}
+	}
+
+	FTangentSolver::FSolverParams::FSolverParams(
+		const FPointerEvent& InMouseEvent,
+		const FCurveEditorScreenSpace& InCurveSpace,
+		const FVector2D& InPixelDelta,
+		const float InDisplayRatio,
+		const bool bInArriveTangent,
+		const float InInitialTangent)
+		: MouseEvent(InMouseEvent)
+		, CurveSpace(InCurveSpace)
+		, PixelDelta(InPixelDelta)
+		, DisplayRatio(InDisplayRatio)
+		, bArriveTangent(bInArriveTangent)
+		, InitialTangent(InInitialTangent)
+	{}
+
+	void FTangentSolver::SolveTangentConstantWeight(const FSolverParams& Params)
+	{
+		if (Params.MouseEvent.IsMouseButtonDown(EKeys::LeftMouseButton))
+		{				
+			const float PixelLength = Params.bArriveTangent ? -60.0f : 60.f;
+			const float Slope = Params.InitialTangent * -Params.DisplayRatio;
+
+			const FVector2D InitialTangentOffset = CurveEditor::GetVectorFromSlopeAndLength(Slope, PixelLength);
+			FVector2D NewTangentOffset = InitialTangentOffset + Params.PixelDelta;
+
+			MendNearlyZeroXTangentOffset(Params.bArriveTangent, NewTangentOffset);
+
+			const float Tangent = (-NewTangentOffset.Y / NewTangentOffset.X) / Params.DisplayRatio;
+
+			if (Params.bArriveTangent)
+			{
+				KeyAttributes.SetArriveTangent(Tangent);
+			}
+			else
+			{
+				KeyAttributes.SetLeaveTangent(Tangent);
+			}
+
+			LastTangentOffset = NewTangentOffset;
+		}
+	}
+
+	void FTangentSolver::SolveTangentMutableWeight(const FSolverParams& Params, float InitialWeight)
+	{
+		InitialWeight = Params.bArriveTangent ? -InitialWeight : InitialWeight;
+
+		FVector2D InitialTangentOffset = CurveEditor::ComputeScreenSpaceTangentOffset(
+			Params.CurveSpace,
+			Params.InitialTangent,
+			InitialWeight);
+
+		FVector2D NewTangentOffset;
+		bool bValueSet = false;
+
+		if (Params.MouseEvent.IsMouseButtonDown(EKeys::LeftMouseButton))
+		{
+			const FVector2D PreviousTangentOffset = LastTangentOffset.IsSet() ? LastTangentOffset.GetValue() : InitialTangentOffset;
+			const FVector2D PointAtOffset = InitialTangentOffset + Params.PixelDelta;
+
+			if (Params.MouseEvent.IsShiftDown() &&
+				Params.MouseEvent.IsControlDown())
+			{
+				// Shift-LMB adjusts the weight, but not the tangent
+				const double Scalar = FVector2D::DotProduct(PointAtOffset, PreviousTangentOffset.GetSafeNormal());
+
+				if (Scalar >= 1.0)
+				{
+					NewTangentOffset = PreviousTangentOffset.GetSafeNormal() * Scalar;
+				}
+				else
+				{
+					NewTangentOffset = PreviousTangentOffset.GetSafeNormal();
+				}
+			}
+			else if (Params.MouseEvent.IsShiftDown())
+			{
+				// Ctrl-Shift-LMB adjusts the tangent, but not the weight
+				const double Weigth = PreviousTangentOffset.Size();
+				NewTangentOffset = PointAtOffset.GetSafeNormal() * Weigth;
+
+				if ((Params.bArriveTangent && NewTangentOffset.X >= 0.0) ||
+					(!Params.bArriveTangent && NewTangentOffset.X <= 0.0))
+				{
+					NewTangentOffset = { 0.0, Weigth * FMath::Sign(NewTangentOffset.Y) };
+				}
+			}
+			else
+			{
+				// LMB adjusts the tangent and weight 
+				NewTangentOffset = InitialTangentOffset + Params.PixelDelta;
+			}
+
+			LastTangentOffset = NewTangentOffset;
+
+			MendNearlyZeroXTangentOffset(Params.bArriveTangent, NewTangentOffset);
+
+			float Tangent;
+			float Weight;
+			CurveEditor::TangentAndWeightFromOffset(Params.CurveSpace, NewTangentOffset, Tangent, Weight);
+
+			if (Params.bArriveTangent)
+			{
+				KeyAttributes.SetArriveTangent(Tangent);
+				KeyAttributes.SetArriveTangentWeight(Weight);
+			}
+			else
+			{
+				KeyAttributes.SetLeaveTangent(Tangent);
+				KeyAttributes.SetLeaveTangentWeight(Weight);
+			}
+		}
+	}
+
+	void FTangentSolver::MendNearlyZeroXTangentOffset(const bool bArriveTangent, FVector2D& InOutTangentOffset) const
+	{
+		constexpr float TangentCrossoverThresholdPx = 1.f;
+		if (bArriveTangent)
+		{
+			InOutTangentOffset.X = FMath::Min(InOutTangentOffset.X, -TangentCrossoverThresholdPx);
+		}
+		else
+		{
+			InOutTangentOffset.X = FMath::Max(InOutTangentOffset.X, TangentCrossoverThresholdPx);
+		}
+	}
 }
 
 void FCurveEditorDragOperation_Tangent::OnInitialize(FCurveEditor* InCurveEditor, const TOptional<FCurvePointHandle>& InCardinalPoint)
@@ -59,13 +243,19 @@ void FCurveEditorDragOperation_Tangent::OnBeginDrag(FVector2D InitialPosition, F
 			FKeyData& KeyData = KeysByCurve.Emplace_GetRef(CurveID);
 			KeyData.Handles = TArray<FKeyHandle>(Handles.GetData(), Handles.Num());
 
-			for (FKeyHandle Handle : Handles)
+			TArray<FKeyAttributes> KeyAttributesArray;
+			KeyAttributesArray.SetNum(KeyData.Handles.Num());
+			Curve->GetKeyAttributes(KeyData.Handles, KeyAttributesArray);
+
+			KeyData.TangentSolvers.Reserve(Handles.Num());
+			for (int32 HandleIndex = 0; HandleIndex < Handles.Num(); ++HandleIndex)
 			{
-				KeyData.PointTypes.Add(Pair.Value.PointType(Handle));
-			}
+				const FKeyHandle& Handle = Handles[HandleIndex];
+				const ECurvePointType PointType = Pair.Value.PointType(Handle);
+				const FKeyAttributes& KeyAttributes = KeyAttributesArray[HandleIndex];
 
-			KeyData.Attributes.SetNum(KeyData.Handles.Num());
-			Curve->GetKeyAttributes(KeyData.Handles, KeyData.Attributes);
+				KeyData.TangentSolvers.Add(CurveEditorDragOperation::FTangentSolver(PointType, KeyAttributes));
+			}
 		}
 	}
 }
@@ -75,7 +265,6 @@ void FCurveEditorDragOperation_Tangent::OnDrag(FVector2D InitialPosition, FVecto
 	FVector2D PixelDelta = CurrentPosition - InitialPosition;
 
 	TArray<FKeyAttributes> NewKeyAttributesScratch;
-	TArray<FKeyPosition> KeyPositions;
 
 	for (FKeyData& KeyData : KeysByCurve)
 	{
@@ -92,110 +281,18 @@ void FCurveEditorDragOperation_Tangent::OnDrag(FVector2D InitialPosition, FVecto
 		}
 
 		FCurveEditorScreenSpace CurveSpace = View->GetCurveSpace(KeyData.CurveID);
-		const float DisplayRatio = (CurveSpace.PixelsPerOutput() / CurveSpace.PixelsPerInput());
 
 		NewKeyAttributesScratch.Reset();
-		NewKeyAttributesScratch.Reserve(KeyData.Attributes.Num());
-		KeyPositions.SetNumUninitialized(KeyData.Attributes.Num());
-		Curve->GetKeyPositions(KeyData.Handles, KeyPositions);
+		NewKeyAttributesScratch.Reserve(KeyData.TangentSolvers.Num());
 
-		for (int32 i = 0; i < KeyData.Handles.Num(); ++i)
+		for (CurveEditorDragOperation::FTangentSolver& Solver : KeyData.TangentSolvers)
 		{
-			const FKeyAttributes Attributes = KeyData.Attributes[i];
-
-			FKeyAttributes NewAttributes;
-
-			if (Attributes.HasArriveTangent() && EnumHasAnyFlags(KeyData.PointTypes[i], ECurvePointType::ArriveTangent))
-			{
-				const float ArriveTangent = Attributes.GetArriveTangent();
-				if (Attributes.HasTangentWeightMode() && Attributes.HasArriveTangentWeight() &&
-					(Attributes.GetTangentWeightMode() == RCTWM_WeightedBoth || Attributes.GetTangentWeightMode() == RCTWM_WeightedArrive))
-				{
-					FVector2D TangentOffset = CurveEditor::ComputeScreenSpaceTangentOffset(CurveSpace, ArriveTangent, -Attributes.GetArriveTangentWeight());
-					TangentOffset += PixelDelta;
-
-					if (MouseEvent.IsShiftDown())
-					{
-						TangentOffset = RoundTrajectory(TangentOffset);
-					}
-
-					// We prevent the handle from crossing over the 0 point... the code handles it but it creates an ugly pop in your curve and it lets your Arrive tangents become Leave tangents which defeats the point.
-					TangentOffset.X = FMath::Min(TangentOffset.X, -CurveEditorDragOperation::TangentCrossoverThresholdPx);
-
-					float Tangent, Weight;
-					CurveEditor::TangentAndWeightFromOffset(CurveSpace, TangentOffset, Tangent, Weight);
-
-					NewAttributes.SetArriveTangent(Tangent);
-					NewAttributes.SetArriveTangentWeight(Weight);
-				}
-				else
-				{
-					const float PixelLength = 60.0f;
-					FVector2D TangentOffset = CurveEditor::GetVectorFromSlopeAndLength(ArriveTangent * -DisplayRatio, -PixelLength);
-					TangentOffset += PixelDelta;
-
-					if (MouseEvent.IsShiftDown())
-					{
-						TangentOffset = RoundTrajectory(TangentOffset);
-					}
-
-					// We prevent the handle from crossing over the 0 point... the code handles it but it creates an ugly pop in your curve and it lets your Arrive tangents become Leave tangents which defeats the point.
-					TangentOffset.X = FMath::Min(TangentOffset.X, -CurveEditorDragOperation::TangentCrossoverThresholdPx);
+			Solver.Process(MouseEvent, CurveSpace, PixelDelta);
 
-					const float Tangent = (-TangentOffset.Y / TangentOffset.X) / DisplayRatio;
-
-					NewAttributes.SetArriveTangent(Tangent);
-				}
-			}
-
-			if (Attributes.HasLeaveTangent() && EnumHasAnyFlags(KeyData.PointTypes[i], ECurvePointType::LeaveTangent))
-			{
-				const float LeaveTangent = Attributes.GetLeaveTangent();
-
-				if (Attributes.HasTangentWeightMode() && Attributes.HasLeaveTangentWeight() &&
-					(Attributes.GetTangentWeightMode() == RCTWM_WeightedBoth || Attributes.GetTangentWeightMode() == RCTWM_WeightedLeave))
-				{
-					FVector2D TangentOffset = CurveEditor::ComputeScreenSpaceTangentOffset(CurveSpace, LeaveTangent, Attributes.GetLeaveTangentWeight());
-					TangentOffset += PixelDelta;
-
-					if (MouseEvent.IsShiftDown())
-					{
-						TangentOffset = RoundTrajectory(TangentOffset);
-					}
-
-					// We prevent the handle from crossing over the 0 point... the code handles it but it creates an ugly pop in your curve and it lets your Arrive tangents become Leave tangents which defeats the point.
-					TangentOffset.X = FMath::Max(TangentOffset.X, CurveEditorDragOperation::TangentCrossoverThresholdPx);
-
-					float Tangent, Weight;
-					CurveEditor::TangentAndWeightFromOffset(CurveSpace, TangentOffset, Tangent, Weight);
-
-					NewAttributes.SetLeaveTangent(Tangent);
-					NewAttributes.SetLeaveTangentWeight(Weight);
-				}
-				else
-				{
-					const float PixelLength = 60.0f;
-					FVector2D TangentOffset = CurveEditor::GetVectorFromSlopeAndLength(LeaveTangent * -DisplayRatio, PixelLength);
-					TangentOffset += PixelDelta;
-
-					if (MouseEvent.IsShiftDown())
-					{
-						TangentOffset = RoundTrajectory(TangentOffset);
-					}
-
-					// We prevent the handle from crossing over the 0 point... the code handles it but it creates an ugly pop in your curve and it lets your Arrive tangents become Leave tangents which defeats the point.
-					TangentOffset.X = FMath::Max(TangentOffset.X, CurveEditorDragOperation::TangentCrossoverThresholdPx);
-
-					const float Tangent = (-TangentOffset.Y / TangentOffset.X) / DisplayRatio;
-
-					NewAttributes.SetLeaveTangent(Tangent);
-				}
-			}
-				
-			NewKeyAttributesScratch.Add(NewAttributes);
+			NewKeyAttributesScratch.Add(Solver.GetKeyAttributes());
 		}
+
 		Curve->SetKeyAttributes(KeyData.Handles, NewKeyAttributesScratch, EPropertyChangeType::Interactive);
-		KeyData.LastDraggedAttributes = NewKeyAttributesScratch;
 	}
 }
 
@@ -207,7 +304,11 @@ void FCurveEditorDragOperation_Tangent::OnCancelDrag()
 	{
 		if (FCurveModel* Curve = CurveEditor->FindCurve(KeyData.CurveID))
 		{
-			Curve->SetKeyAttributes(KeyData.Handles, KeyData.Attributes, EPropertyChangeType::ValueSet);
+			TArray<FKeyAttributes> InitialAttributes;
+			InitialAttributes.Reserve(KeyData.TangentSolvers.Num());
+			Algo::Transform(KeyData.TangentSolvers, InitialAttributes, &CurveEditorDragOperation::FTangentSolver::GetInitialKeyAttributes);
+
+			Curve->SetKeyAttributes(KeyData.Handles, InitialAttributes, EPropertyChangeType::ValueSet);
 		}
 	}
 
@@ -222,22 +323,13 @@ void FCurveEditorDragOperation_Tangent::OnEndDrag(FVector2D InitialPosition, FVe
 	{
 		if (FCurveModel* Curve = CurveEditor->FindCurve(KeyData.CurveID))
 		{
-			if (KeyData.LastDraggedAttributes.IsSet())
-			{
-				Curve->SetKeyAttributes(KeyData.Handles, KeyData.LastDraggedAttributes.GetValue(), EPropertyChangeType::ValueSet);
-			}
+			TArray<FKeyAttributes> FinalAttributes;
+			FinalAttributes.Reserve(KeyData.TangentSolvers.Num());
+			Algo::Transform(KeyData.TangentSolvers, FinalAttributes, &CurveEditorDragOperation::FTangentSolver::GetKeyAttributes);
+
+			Curve->SetKeyAttributes(KeyData.Handles, FinalAttributes, EPropertyChangeType::ValueSet);
 		}
 	}
 	
 	CurveEditor->SuppressBoundTransformUpdates(false);
 }
-
-FVector2D FCurveEditorDragOperation_Tangent::RoundTrajectory(FVector2D Delta)
-{
-	float Distance = Delta.Size();
-	float Theta = FMath::Atan2(Delta.Y, Delta.X) + PI/2;
-
-	float RoundTo = PI / 4;
-	Theta = FMath::RoundToInt(Theta / RoundTo) * RoundTo - PI/2;
-	return FVector2D(Distance * FMath::Cos(Theta), Distance * FMath::Sin(Theta));
-}
\ No newline at end of file
diff --git a/Engine/Source/Editor/CurveEditor/Private/DragOperations/CurveEditorDragOperation_Tangent.h b/Engine/Source/Editor/CurveEditor/Private/DragOperations/CurveEditorDragOperation_Tangent.h
--- a/Engine/Source/Editor/CurveEditor/Private/DragOperations/CurveEditorDragOperation_Tangent.h
+++ b/Engine/Source/Editor/CurveEditor/Private/DragOperations/CurveEditorDragOperation_Tangent.h
@@ -8,11 +8,97 @@
 #include "Curves/KeyHandle.h"
 #include "ICurveEditorDragOperation.h"
 #include "Math/Vector2D.h"
+#include "Misc/EnumClassFlags.h"
 #include "Misc/Optional.h"
 
 class FCurveEditor;
+struct FCurveEditorScreenSpace;
 struct FPointerEvent;
 
+namespace CurveEditorDragOperation
+{
+	/** Solves the tangent offset for a single key depending on mouse button and modifier keys */
+	class FTangentSolver
+	{
+	public:
+		FTangentSolver(ECurvePointType InPointType, FKeyAttributes InInitialKeyAttributes);
+
+		/** 
+		 * Solves the tangent considering the current mouse position, the mouse buttons and modifier keys
+		 * 
+		 * @param MouseEvent			The mouse event to Process
+		 * @param CurveSpace			The screen space of the curve editor
+		 * @param PixelDelta			The mouse position delta from the initial mouse position in pixels
+		 */
+		void Process(const FPointerEvent& MouseEvent, const FCurveEditorScreenSpace& CurveSpace, const FVector2D& PixelDelta);
+
+		/** Returns the current key attributes */
+		const FKeyAttributes& GetKeyAttributes() const { return KeyAttributes; }
+
+		/** Returns the initial key attributes */
+		const FKeyAttributes& GetInitialKeyAttributes() const { return InitialKeyAttributes; }
+
+	private:
+		/** Struct that holds relevant data to compute the new tangent */
+		struct FSolverParams
+		{
+			FSolverParams(
+				const FPointerEvent& InMouseEvent,
+				const FCurveEditorScreenSpace& InCurveSpace,
+				const FVector2D& InPixelDelta,
+				const float InDisplayRatio,
+				const bool bInArriveTangent,
+				const float InInitialTangent);
+
+			/** The current mouse event */
+			const FPointerEvent& MouseEvent;
+
+			/** The current curve space */
+			const FCurveEditorScreenSpace& CurveSpace;
+
+			/** The delta of the mouse cursor from the initial position in pixels */
+			const FVector2D& PixelDelta;
+
+			/** The current display ratio */
+			const float DisplayRatio;
+
+			/** True if the tangent to process is an arrive tangent, else it is a leave tangent */
+			const bool bArriveTangent;
+
+			/** The initial tangent when the drag drop op started */
+			const float InitialTangent;
+
+			// Non-copyable
+			FSolverParams(const FSolverParams&) = delete;
+			FSolverParams& operator=(const FSolverParams&) = delete;
+		};
+
+		/** Solves tangent leaving the weight untouched. Useful when tangent weight mode is disabled. */
+		void SolveTangentConstantWeight(const FSolverParams& Params);
+
+		/** Solves tangent and weight. Useful when tangent weight mode is enabled. */
+		void SolveTangentMutableWeight(const FSolverParams& Params, float InitialWeight);
+
+		/** 
+		 * Prevents the handle from crossing over the 0 point. Curve editor would handle it but it creates 
+		 * an ugly pop in the curve and it lets the Arrive tangents become Leave tangents which defeats the point.
+		 */
+		void MendNearlyZeroXTangentOffset(const bool bArriveTangent, FVector2D& InOutTangentOffset) const;
+
+		/** Holds the last tanget offset. Useful when editing the tangent or weight exclusively. */
+		TOptional<FVector2D> LastTangentOffset;
+
+		/** The current attributes */
+		FKeyAttributes KeyAttributes;
+
+		/** The initial attributes of the key */
+		const FKeyAttributes InitialKeyAttributes;
+
+		/** The point type of the key handle for which the tangent is solved */
+		const ECurvePointType PointType; 
+	};
+}
+
 class FCurveEditorDragOperation_Tangent : public ICurveEditorKeyDragOperation
 {
 public:
@@ -23,13 +109,7 @@ public:
 	virtual void OnEndDrag(FVector2D InitialPosition, FVector2D CurrentPosition, const FPointerEvent& MouseEvent) override;
 	virtual void OnCancelDrag() override;
 
-private:
-
-	/** Round a trajectory to the nearest 45 degrees */
-	static FVector2D RoundTrajectory(FVector2D Delta);
-
-private:
-
+private:		
 	/** Ptr back to the curve editor */
 	FCurveEditor* CurveEditor;
 
@@ -43,14 +123,12 @@ private:
 
 		/** The curve that contains the keys we're dragging */
 		FCurveModelID CurveID;
-		/** All the handles within a given curve that we are dragging */
+
+		/** The handles that are being dragged */
 		TArray<FKeyHandle> Handles;
-		/** All the point types within a given curve that we are dragging */
-		TArray<ECurvePointType> PointTypes;
-		/** The key attributes for each of the above handles */
-		TArray<FKeyAttributes> Attributes;
-		/** Used in OnEndDrag to send final key updates */
-		TOptional<TArray<FKeyAttributes> > LastDraggedAttributes;
+
+		/** Tangent solvers for the handles that are being dragged */
+		TArray<CurveEditorDragOperation::FTangentSolver> TangentSolvers;
 	};
 
 	/** Key dragging data stored per-curve */
