Curve Editor: Zoom no longer resets when zooming, paning or editing keys in Normalized View Mode

diff --git a/Engine/Source/Editor/CurveEditor/Private/Views/SCurveEditorViewNormalized.cpp b/Engine/Source/Editor/CurveEditor/Private/Views/SCurveEditorViewNormalized.cpp
--- a/Engine/Source/Editor/CurveEditor/Private/Views/SCurveEditorViewNormalized.cpp
+++ b/Engine/Source/Editor/CurveEditor/Private/Views/SCurveEditorViewNormalized.cpp
@@ -181,9 +181,14 @@ void SCurveEditorViewNormalized::UpdateViewToTransformCurves(double InputMin, do
 		{
 			continue;
 		}
+		
+		// Consider the whole time range when getting the value range
+		double MinTime;
+		double MaxTime;
+		Curve->GetTimeRange(MinTime, MaxTime);
 
 		double CurveOutputMin = 0, CurveOutputMax = 1;
-		Curve->GetValueRange(InputMin, InputMax, CurveOutputMin, CurveOutputMax);
+		Curve->GetValueRange(MinTime, MaxTime, CurveOutputMin, CurveOutputMax);
 
 		It->Value.ViewToCurveTransform = CalculateViewToCurveTransform(OutputMin, OutputMax, CurveOutputMin, CurveOutputMax);
 	}
@@ -192,19 +197,30 @@ void SCurveEditorViewNormalized::UpdateViewToTransformCurves(double InputMin, do
 
 void SCurveEditorViewNormalized::Tick(const FGeometry& AllottedGeometry, const double InCurrentTime, const float InDeltaTime)
 {
-	TSharedPtr<FCurveEditor> CurveEditor = WeakCurveEditor.Pin();
-	if (!CurveEditor)
+	const TSharedPtr<FCurveEditor> CurveEditor = WeakCurveEditor.Pin();
+	if (CurveEditor.IsValid() && !CurveEditor->AreBoundTransformUpdatesSuppressed())
 	{
-		return;
+		InternalUpdateViewToTransformCurves();
 	}
 
-	if (!CurveEditor->AreBoundTransformUpdatesSuppressed())
+	SInteractiveCurveEditorView::Tick(AllottedGeometry, InCurrentTime, InDeltaTime);
+}
+
+FReply SCurveEditorViewNormalized::OnMouseButtonUp(const FGeometry& MyGeometry, const FPointerEvent& MouseEvent)
+{
+	InternalUpdateViewToTransformCurves();
+
+	return SInteractiveCurveEditorView::OnMouseButtonUp(MyGeometry, MouseEvent);
+}
+
+void SCurveEditorViewNormalized::InternalUpdateViewToTransformCurves()
+{
+	const TSharedPtr<FCurveEditor> CurveEditor = WeakCurveEditor.Pin();
+	if (CurveEditor.IsValid())
 	{
 		bFixedOutputBounds = false;
 		double InputMin = 0, InputMax = 1;
 		GetInputBounds(InputMin, InputMax);
 		UpdateViewToTransformCurves(InputMin, InputMax);
 	}
-
-	SInteractiveCurveEditorView::Tick(AllottedGeometry, InCurrentTime, InDeltaTime);
-}
\ No newline at end of file
+}
diff --git a/Engine/Source/Editor/CurveEditor/Public/Views/SCurveEditorViewNormalized.h b/Engine/Source/Editor/CurveEditor/Public/Views/SCurveEditorViewNormalized.h
--- a/Engine/Source/Editor/CurveEditor/Public/Views/SCurveEditorViewNormalized.h
+++ b/Engine/Source/Editor/CurveEditor/Public/Views/SCurveEditorViewNormalized.h
@@ -34,9 +34,15 @@ public:
 
 protected:
 
+	// SWidget Interface
 	virtual void Tick(const FGeometry& AllottedGeometry, const double InCurrentTime, const float InDeltaTime) override;
+	virtual FReply OnMouseButtonUp(const FGeometry& MyGeometry, const FPointerEvent& MouseEvent) override;
+	// ~SWidget Interface
 
 	virtual void PaintView(const FPaintArgs& Args, const FGeometry& AllottedGeometry, const FSlateRect& MyCullingRect, FSlateWindowElementList& OutDrawElements, int32 BaseLayerId, const FWidgetStyle& InWidgetStyle, bool bParentEnabled) const override;
 
 	virtual void DrawBufferedCurves(const FGeometry& AllottedGeometry, const FSlateRect& MyCullingRect, FSlateWindowElementList& OutDrawElements, int32 BaseLayerId, ESlateDrawEffect DrawEffects) const;
+
+	/** Helper to call UpdateViewToTransformCurves from internal methods */
+	void InternalUpdateViewToTransformCurves();
 };
