Curve Editor: Fixed minor grid line rendering in the wrong space

diff --git a/Engine/Source/Editor/CurveEditor/Private/Views/SInteractiveCurveEditorView.cpp b/Engine/Source/Editor/CurveEditor/Private/Views/SInteractiveCurveEditorView.cpp
--- a/Engine/Source/Editor/CurveEditor/Private/Views/SInteractiveCurveEditorView.cpp
+++ b/Engine/Source/Editor/CurveEditor/Private/Views/SInteractiveCurveEditorView.cpp
@@ -545,6 +545,7 @@ void SInteractiveCurveEditorView::DrawGridLines(TSharedRef<FCurveEditor> CurveEd
 	// Now draw the minor vertical lines which are drawn with a lighter color.
 	for (float PosX : MinorGridLines)
 	{
+		PosX = HorizontalGridSpace.SecondsToScreen(PosX);
 		if (PosX < 0 || PosX > Width)
 		{
 			continue;
@@ -761,6 +762,7 @@ void SInteractiveCurveEditorView::DrawGridLines(TSharedRef<FCurveEditor> CurveEd
 	// Draw our minor horizontal lines
 	for (float PosY : MinorGridLines)
 	{
+		PosY = VerticalGridSpace.ValueToScreen(PosY);
 		if (PosY < 0 || PosY > Height)
 		{
 			continue;
