Curve Editor: Allow the selection to be serialized (e.g. for undo / redo)

diff --git a/Engine/Source/Editor/CurveEditor/Private/CurveEditorSelection.cpp b/Engine/Source/Editor/CurveEditor/Private/CurveEditorSelection.cpp
--- a/Engine/Source/Editor/CurveEditor/Private/CurveEditorSelection.cpp
+++ b/Engine/Source/Editor/CurveEditor/Private/CurveEditorSelection.cpp
@@ -27,6 +27,12 @@ void FCurveEditorSelection::Clear()
 	++SerialNumber;
 }
 
+void FCurveEditorSelection::SerializeSelection(FArchive& Ar)
+{
+	Ar << SerialNumber;
+	Ar << CurveToSelectedKeys;
+}
+
 const FKeyHandleSet* FCurveEditorSelection::FindForCurve(FCurveModelID InCurveID) const
 {
 	return CurveToSelectedKeys.Find(InCurveID);
diff --git a/Engine/Source/Editor/CurveEditor/Public/CurveEditorSelection.h b/Engine/Source/Editor/CurveEditor/Public/CurveEditorSelection.h
--- a/Engine/Source/Editor/CurveEditor/Public/CurveEditorSelection.h
+++ b/Engine/Source/Editor/CurveEditor/Public/CurveEditorSelection.h
@@ -58,6 +58,16 @@ struct CURVEEDITOR_API FKeyHandleSet
 	 */
 	ECurvePointType PointType(FKeyHandle Handle) const { return HandleToPointType.FindChecked(Handle); }
 
+	/**
+	 * Serializes the handle set
+	 */
+	friend FArchive& operator<<(FArchive& Ar, FKeyHandleSet& HandleSet)
+	{
+		Ar << HandleSet.SortedHandles;
+		Ar << HandleSet.HandleToPointType;
+		return Ar;
+	}
+
 private:
 
 	/** Sorted array of key handles */
@@ -180,6 +190,11 @@ public:
 	 */
 	void Clear();
 
+public:
+
+	/** Serializes the SerialNumber and CurveToSelectedKeys. */
+	void SerializeSelection(FArchive& Ar);
+
 private:
 
 	/** Weak reference to the curve editor to check whether keys are locked or not */
diff --git a/Engine/Source/Editor/CurveEditor/Public/CurveEditorTypes.h b/Engine/Source/Editor/CurveEditor/Public/CurveEditorTypes.h
--- a/Engine/Source/Editor/CurveEditor/Public/CurveEditorTypes.h
+++ b/Engine/Source/Editor/CurveEditor/Public/CurveEditorTypes.h
@@ -47,11 +47,19 @@ ENUM_CLASS_FLAGS(ECurveEditorViewID);
  */
 struct CURVEEDITOR_API FCurveModelID
 {
+	/** Default constructor is needed by e.g. TArray::operator<< */
+	FCurveModelID() : ID(0) {}
+
 	/**
 	 * Generate a new curve model ID
 	 */
 	static FCurveModelID Unique();
 
+	/**
+	 * @return Whether this ID is valid (was constructed with Unique)
+	 */
+	bool IsValid() const { return ID != 0; }
+
 	/**
 	 * Check two IDs for equality
 	 */
@@ -84,8 +92,16 @@ struct CURVEEDITOR_API FCurveModelID
 		return GetTypeHash(In.ID);
 	}
 
+	/**
+	 * Save a curve model ID
+	 */
+	FORCEINLINE friend FArchive& operator<<(FArchive& Ar, FCurveModelID& Handle)
+	{
+		Ar << Handle.ID;
+		return Ar;
+	}
+
 private:
-	FCurveModelID() {}
 
 	/** Internal serial ID */
 	uint32 ID;
