Curve Editor: Surface Euler button into toolbar:
- New FPromotedFilterContainer: Keeps track of filters promoted to toolbar and creates dynamic FUICommandInfos
- New FPromotedFilterCommandBinder: Listens for changes to promoted filters and binds commands
- ICurveEditorModule now exposes the global FPromotedFilterContainer which affects promoted filters across all editors
- FCurveEditor::GetToolbarPromotedFilters is the access point for UI.
- The look in toolbar is determined by the dynamic FUICommandInfos, which get the info from static getters on UCurveEditorFilterBase.
- Missing: regenerating toolbar widget when list of promoted filters changes (do that when UE-230269 is implemented)

diff --git a/Engine/Source/Editor/CurveEditor/Private/CurveEditor.cpp b/Engine/Source/Editor/CurveEditor/Private/CurveEditor.cpp
--- a/Engine/Source/Editor/CurveEditor/Private/CurveEditor.cpp
+++ b/Engine/Source/Editor/CurveEditor/Private/CurveEditor.cpp
@@ -557,6 +557,15 @@ void FCurveEditor::BindCommands()
 	}
 }
 
+TSharedPtr<UE::CurveEditor::FPromotedFilterContainer> FCurveEditor::GetToolbarPromotedFilters() const
+{
+	TSharedPtr<UE::CurveEditor::FPromotedFilterContainer> Result = ICurveEditorModule::Get().GetGlobalToolbarPromotedFilters();
+	checkf(Result, TEXT("Should be valid for the lifetime of the module"));
+
+	// In the future, we could extend FCurveEditor to have its own override for the globally promoted filters.
+	return Result;
+}
+
 FCurveSnapMetrics FCurveEditor::GetCurveSnapMetrics(FCurveModelID CurveModel) const
 {
 	FCurveSnapMetrics CurveMetrics;
diff --git a/Engine/Source/Editor/CurveEditor/Private/CurveEditorContextMenu.cpp b/Engine/Source/Editor/CurveEditor/Private/CurveEditorContextMenu.cpp
--- a/Engine/Source/Editor/CurveEditor/Private/CurveEditorContextMenu.cpp
+++ b/Engine/Source/Editor/CurveEditor/Private/CurveEditorContextMenu.cpp
@@ -10,6 +10,7 @@
 #include "CurveEditorTypes.h"
 #include "CurveModel.h"
 #include "Delegates/Delegate.h"
+#include "Filters/PromotedFilterContainer.h"
 #include "Framework/Commands/GenericCommands.h"
 #include "Framework/MultiBox/MultiBoxBuilder.h"
 #include "HAL/Platform.h"
@@ -133,6 +134,7 @@ void FCurveEditorContextMenu::BuildMenu(FMenuBuilder& MenuBuilder, TSharedRef<FC
 
 			// Filters
 			MenuBuilder.AddMenuEntry(FCurveEditorCommands::Get().OpenUserImplementableFilterWindow);
+			CurveEditor->GetToolbarPromotedFilters()->AppendToBuilder(MenuBuilder);
 			
 			// View
 			MenuBuilder.AddMenuEntry(FCurveEditorCommands::Get().ZoomToFit);
@@ -207,6 +209,7 @@ void FCurveEditorContextMenu::BuildMenu(FMenuBuilder& MenuBuilder, TSharedRef<FC
 
 				// Filters
 				MenuBuilder.AddMenuEntry(FCurveEditorCommands::Get().OpenUserImplementableFilterWindow);
+				CurveEditor->GetToolbarPromotedFilters()->AppendToBuilder(MenuBuilder);
 
 				// View
 				MenuBuilder.AddMenuEntry(FCurveEditorCommands::Get().ZoomToFit);
@@ -239,6 +242,7 @@ void FCurveEditorContextMenu::BuildMenu(FMenuBuilder& MenuBuilder, TSharedRef<FC
 
 				// Filters
 				MenuBuilder.AddMenuEntry(FCurveEditorCommands::Get().OpenUserImplementableFilterWindow);
+				CurveEditor->GetToolbarPromotedFilters()->AppendToBuilder(MenuBuilder);
 				
 				// View
 				MenuBuilder.AddMenuEntry(FCurveEditorCommands::Get().ZoomToFit);
diff --git a/Engine/Source/Editor/CurveEditor/Private/CurveEditorModule.cpp b/Engine/Source/Editor/CurveEditor/Private/CurveEditorModule.cpp
--- a/Engine/Source/Editor/CurveEditor/Private/CurveEditorModule.cpp
+++ b/Engine/Source/Editor/CurveEditor/Private/CurveEditorModule.cpp
@@ -9,27 +9,31 @@
 #include "Delegates/Delegate.h"
 #include "Filters/CurveEditorBakeFilter.h"
 #include "Filters/CurveEditorBakeFilterCustomization.h"
+#include "Filters/CurveEditorEulerFilter.h"
+#include "Filters/PromotedFilterContainer.h"
 #include "HAL/PlatformCrt.h"
 #include "ICurveEditorModule.h"
 #include "Misc/CoreDelegates.h"
 #include "Modules/ModuleManager.h"
 #include "PropertyEditorModule.h"
+#include "Templates/SharedPointer.h"
 #include "ToolMenus.h"
 
 class FCurveEditorModule : public ICurveEditorModule
 {
 public:
+	
 	virtual void StartupModule() override
 	{
 		if (GIsEditor)
 		{
 			if (UToolMenus::TryGet())
 			{
-				FCurveEditorCommands::Register();
+				RegisterCommands();
 			}
 			else
 			{
-				FCoreDelegates::OnPostEngineInit.AddStatic(&FCurveEditorCommands::Register);
+				FCoreDelegates::OnPostEngineInit.AddRaw(this, &FCurveEditorModule::RegisterCommands);
 			}
 		}
 
@@ -41,6 +45,7 @@ public:
 		FCurveEditorCommands::Unregister();
 
 		UnregisterCustomizations();
+		ToolbarPromotedFilters.Reset();
 	}
 
 	virtual FDelegateHandle RegisterEditorExtension(FOnCreateCurveEditorExtension InOnCreateCurveEditorExtension) override
@@ -94,6 +99,8 @@ public:
 		return ToolExtensionDelegates;
 	}
 
+	virtual TSharedPtr<UE::CurveEditor::FPromotedFilterContainer> GetGlobalToolbarPromotedFilters() const override { return ToolbarPromotedFilters; }
+
 	void RegisterCustomizations()
 	{
 		FPropertyEditorModule& PropertyEditorModule = FModuleManager::Get().LoadModuleChecked<FPropertyEditorModule>("PropertyEditor");
@@ -114,6 +121,7 @@ public:
 	}
 
 private:
+	
 	/** List of editor extension handler delegates Curve Editors will execute when they are created. */
 	TArray<FOnCreateCurveEditorExtension> EditorExtensionDelegates;
 
@@ -123,6 +131,25 @@ private:
 	/** List of Extenders that that should be called when building the Curve Editor toolbar. */
 	TArray<FCurveEditorMenuExtender> ToolBarMenuExtenders;
 
+	/**
+	 * Keeps track of the filters that are promoted to the toolbar.
+	 * 
+	 * This object dynamically creates FUICommandInfos based on the filters surfaced so it must have a globally unique context name for its
+	 * FBindingContext. FCurveEditor instances will reference this by default (but can theoretically create their own).
+	 *
+	 * Until UE-230269 is implemented, the only filter that is surfaced to this object is the Euler filter. 
+	 */
+	TSharedPtr<UE::CurveEditor::FPromotedFilterContainer> ToolbarPromotedFilters;
+
+	void RegisterCommands()
+	{
+		FCurveEditorCommands::Register();
+
+		// This needs to be created after registering the commands because FCurveEditorCommands is used a parent context
+		ToolbarPromotedFilters = MakeShared<UE::CurveEditor::FPromotedFilterContainer>(TEXT("ToolbarPromotedCurveEditorFilters"));
+		// By default, we surface the Euler filter to the toolbar because it is a common action.
+		ToolbarPromotedFilters->AddInstance(*GetMutableDefault<UCurveEditorEulerFilter>());
+	}
 };
 
 IMPLEMENT_MODULE(FCurveEditorModule, CurveEditor)
\ No newline at end of file
diff --git a/Engine/Source/Editor/CurveEditor/Private/Filters/CurveEditorEulerFilter.cpp b/Engine/Source/Editor/CurveEditor/Private/Filters/CurveEditorEulerFilter.cpp
--- a/Engine/Source/Editor/CurveEditor/Private/Filters/CurveEditorEulerFilter.cpp
+++ b/Engine/Source/Editor/CurveEditor/Private/Filters/CurveEditorEulerFilter.cpp
@@ -2,6 +2,7 @@
 
 #include "Filters/CurveEditorEulerFilter.h"
 
+#include "Algo/AnyOf.h"
 #include "Containers/Array.h"
 #include "Containers/ArrayView.h"
 #include "Containers/UnrealString.h"
@@ -88,3 +89,16 @@ void UCurveEditorEulerFilter::ApplyFilter_Impl(TSharedRef<FCurveEditor> InCurveE
 	}
 	OutKeysToSelect = InKeysToOperateOn;
 }
+
+bool UCurveEditorEulerFilter::CanApplyFilter_Impl(TSharedRef<FCurveEditor> InCurveEditor)
+{
+	const TMap<FCurveModelID, FKeyHandleSet>& Selection = InCurveEditor->GetSelection().GetAll();
+	const bool bHasRotationCurves = !Selection.IsEmpty()
+		&& Algo::AnyOf(InCurveEditor->GetSelection().GetAll(),
+		[&InCurveEditor](const TPair<FCurveModelID, FKeyHandleSet>& Pair)
+		{
+			const FCurveModel* Curve = InCurveEditor->FindCurve(Pair.Key);
+			return Curve && Curve->GetIntentionName().Contains("Rotation");
+		});
+	return bHasRotationCurves;
+}
diff --git a/Engine/Source/Editor/CurveEditor/Private/Filters/CurveEditorFilterBase.cpp b/Engine/Source/Editor/CurveEditor/Private/Filters/CurveEditorFilterBase.cpp
--- a/Engine/Source/Editor/CurveEditor/Private/Filters/CurveEditorFilterBase.cpp
+++ b/Engine/Source/Editor/CurveEditor/Private/Filters/CurveEditorFilterBase.cpp
@@ -1,10 +1,34 @@
 // Copyright Epic Games, Inc. All Rights Reserved.
 
 #include "Filters/CurveEditorFilterBase.h"
+
 #include "CurveEditor.h"
+#include "Internationalization/Text.h"
+#include "Styling/SlateIconFinder.h"
+#include "Templates/SubclassOf.h"
+
+#define LOCTEXT_NAMESPACE "CurveEditorFilterBase"
 
+FText UCurveEditorFilterBase::GetLabel(const TSubclassOf<UCurveEditorFilterBase>& InClass)
+{
+	const FString* Label = InClass ? InClass->FindMetaData(TEXT("CurveEditorLabel")) : nullptr;
+	return Label ? FText::FromString(*Label) : LOCTEXT("Filter.LabelEmpty", "Filter");
+}
+
+FText UCurveEditorFilterBase::GetDescription(const TSubclassOf<UCurveEditorFilterBase>& InClass)
+{
+	const FString* Description = InClass ? InClass->FindMetaData(TEXT("CurveEditorDescription")) : nullptr;
+	return Description ? FText::FromString(*Description) : FText::GetEmpty();
+}
+
+FSlateIcon UCurveEditorFilterBase::GetIcon(const TSubclassOf<UCurveEditorFilterBase>& InClass)
+{
+	return FSlateIconFinder::FindIconForClass(InClass.Get());
+}
 
 bool UCurveEditorFilterBase::CanApplyFilter_Impl(TSharedRef<FCurveEditor> InCurveEditor)
 {
 	return InCurveEditor->GetSelection().Count() > 0;
 }
+
+#undef LOCTEXT_NAMESPACE
\ No newline at end of file
diff --git a/Engine/Source/Editor/CurveEditor/Private/Filters/FilterUtils.cpp b/Engine/Source/Editor/CurveEditor/Private/Filters/FilterUtils.cpp
new file mode 100644
--- /dev/null
+++ b/Engine/Source/Editor/CurveEditor/Private/Filters/FilterUtils.cpp
@@ -0,0 +1,34 @@
+﻿// Copyright Epic Games, Inc. All Rights Reserved.
+
+#include "FilterUtils.h"
+
+#include "Filters/CurveEditorFilterBase.h"
+
+#include "Containers/Map.h"
+#include "Templates/SharedPointer.h"
+
+namespace UE::CurveEditor::FilterUtils
+{
+	void ApplyFilter(const TSharedRef<FCurveEditor>& InCurveEditor, UCurveEditorFilterBase& InFilter)
+	{
+		ApplyFilter(InCurveEditor, InFilter, InCurveEditor->Selection.GetAll());
+	}
+
+	void ApplyFilter(
+		const TSharedRef<FCurveEditor>& InCurveEditor,
+		UCurveEditorFilterBase& InFilter,
+		const TMap<FCurveModelID, FKeyHandleSet>& InSelectedKeys
+		)
+	{
+		TMap<FCurveModelID, FKeyHandleSet> OutKeysToSelect;
+		InFilter.ApplyFilter(InCurveEditor, InSelectedKeys, OutKeysToSelect);
+
+		// Clear their selection and then set it to the keys the filter thinks you should have selected.
+		InCurveEditor->GetSelection().Clear();
+
+		for (const TTuple<FCurveModelID, FKeyHandleSet>& OutSet : OutKeysToSelect)
+		{
+			InCurveEditor->GetSelection().Add(OutSet.Key, ECurvePointType::Key, OutSet.Value.AsArray());
+		}
+	}
+}
diff --git a/Engine/Source/Editor/CurveEditor/Private/Filters/FilterUtils.h b/Engine/Source/Editor/CurveEditor/Private/Filters/FilterUtils.h
new file mode 100644
--- /dev/null
+++ b/Engine/Source/Editor/CurveEditor/Private/Filters/FilterUtils.h
@@ -0,0 +1,25 @@
+﻿// Copyright Epic Games, Inc. All Rights Reserved.
+
+#pragma once
+
+#include "Containers/ContainersFwd.h"
+#include "Templates/SharedPointerFwd.h"
+
+class FCurveEditor;
+class UCurveEditorFilterBase;
+
+struct FCurveModelID;
+struct FKeyHandleSet;
+
+namespace UE::CurveEditor::FilterUtils
+{
+	/** Applies the filter to the entire user selection. */
+	void ApplyFilter(const TSharedRef<FCurveEditor>& InCurveEditor, UCurveEditorFilterBase& InFilter);
+	
+	/** Applies the filter to the supplied keys. */
+	void ApplyFilter(
+		const TSharedRef<FCurveEditor>& InCurveEditor,
+		UCurveEditorFilterBase& InFilter,
+		const TMap<FCurveModelID, FKeyHandleSet>& InSelectedKeys
+		);
+}
diff --git a/Engine/Source/Editor/CurveEditor/Private/Filters/PromotedFilterCommandBinder.cpp b/Engine/Source/Editor/CurveEditor/Private/Filters/PromotedFilterCommandBinder.cpp
new file mode 100644
--- /dev/null
+++ b/Engine/Source/Editor/CurveEditor/Private/Filters/PromotedFilterCommandBinder.cpp
@@ -0,0 +1,108 @@
+﻿// Copyright Epic Games, Inc. All Rights Reserved.
+
+#include "Filters/PromotedFilterCommandBinder.h"
+
+#include "Filters/FilterUtils.h"
+#include "Filters/PromotedFilterContainer.h"
+#include "Framework/Commands/UICommandInfo.h"
+#include "Framework/Commands/UICommandList.h"
+#include "ScopedTransaction.h"
+
+#define LOCTEXT_NAMESPACE "FPromotedFilterCommandBinder"
+
+namespace UE::CurveEditor
+{
+	FPromotedFilterCommandBinder::FPromotedFilterCommandBinder(
+		const TSharedRef<FPromotedFilterContainer>& InContainer,
+		const TSharedRef<FUICommandList>& InCommandList,
+		const TSharedRef<FCurveEditor>& InCurveEditor
+		)
+		: Container(InContainer)
+		, CommandList(InCommandList)
+		, CurveEditor(InCurveEditor)
+	{
+		InContainer->OnFilterAdded().AddRaw(this, &FPromotedFilterCommandBinder::OnFilterAdded);
+		InContainer->OnFilterRemoved().AddRaw(this, &FPromotedFilterCommandBinder::OnFilterRemoved);
+
+		InContainer->ForEachFilter([this, &InCommandList](UCurveEditorFilterBase& InFilter, const TSharedRef<FUICommandInfo>& Command)
+		{
+			MapAction(InFilter, Command, *InCommandList);
+		});
+	}
+
+	FPromotedFilterCommandBinder::~FPromotedFilterCommandBinder()
+	{
+		const TSharedPtr<FPromotedFilterContainer> ContainerPin = Container.Pin();
+		if (ContainerPin)
+		{
+			ContainerPin->OnFilterAdded().RemoveAll(this);
+			ContainerPin->OnFilterRemoved().RemoveAll(this);
+		}
+		
+		const TSharedPtr<FUICommandList> CommandListPin = CommandList.Pin();
+		if (ContainerPin && CommandListPin)
+		{
+			ContainerPin->ForEachFilter([&CommandListPin](UCurveEditorFilterBase&, const TSharedRef<FUICommandInfo>& Command)
+			{
+				CommandListPin->UnmapAction(Command);
+			});
+		}
+	}
+	
+	void FPromotedFilterCommandBinder::OnFilterAdded(UCurveEditorFilterBase& InFilter, const TSharedRef<FUICommandInfo>& InCommand) const
+	{
+		const TSharedPtr<FUICommandList> CommandListPin = CommandList.Pin();
+		if (ensure(CommandListPin))
+		{
+			MapAction(InFilter, InCommand, *CommandListPin);
+		}
+	}
+
+	void FPromotedFilterCommandBinder::OnFilterRemoved(UCurveEditorFilterBase& InFilter, const TSharedRef<FUICommandInfo>& InCommand) const
+	{
+		const TSharedPtr<FUICommandList> CommandListPin = CommandList.Pin();
+		if (ensure(CommandListPin))
+		{
+			CommandListPin->UnmapAction(InCommand);
+		}
+	}
+
+	void FPromotedFilterCommandBinder::MapAction(
+		UCurveEditorFilterBase& InFilter, const TSharedRef<FUICommandInfo>& InCommand, FUICommandList& CommandListPin
+		) const
+	{
+		CommandListPin.MapAction(InCommand, FUIAction(
+			FExecuteAction::CreateRaw(this, &FPromotedFilterCommandBinder::ApplyFilter, &InFilter, InCommand.ToWeakPtr()),
+			FCanExecuteAction::CreateRaw(this, &FPromotedFilterCommandBinder::CanApplyFilter, &InFilter, InCommand.ToWeakPtr())
+			));
+	}
+
+	void FPromotedFilterCommandBinder::ApplyFilter(UCurveEditorFilterBase* InFilter, TWeakPtr<FUICommandInfo> Command) const
+	{
+		check(InFilter);
+	
+		// Unlikely but: if CommandPin is null, that means somebody called RemoveInstance while a menu was open.
+		const TSharedPtr<FUICommandInfo> CommandPin = Command.Pin();
+		const TSharedPtr<FCurveEditor> Editor = CurveEditor.Pin();
+		if (CommandPin && Editor && ensure(InFilter))
+		{
+			const FScopedTransaction Transaction(FText::Format(LOCTEXT("ApplyFmt", "Apply {0}"), CommandPin->GetLabel()));
+			FilterUtils::ApplyFilter(Editor.ToSharedRef(), *InFilter);	
+		}
+	}
+
+	bool FPromotedFilterCommandBinder::CanApplyFilter(UCurveEditorFilterBase* InFilter, TWeakPtr<FUICommandInfo> Command) const
+	{
+		check(InFilter);
+	
+		// Unlikely but: if CommandPin is null, that means somebody called RemoveInstance while a menu was open.
+		if (const TSharedPtr<FUICommandInfo> CommandPin = Command.Pin())
+		{
+			TSharedPtr<FCurveEditor> Editor = CurveEditor.Pin();
+			return ensure(Editor && InFilter) && InFilter->CanApplyFilter(Editor.ToSharedRef());
+		}
+		return false;
+	}
+}
+
+#undef LOCTEXT_NAMESPACE
\ No newline at end of file
diff --git a/Engine/Source/Editor/CurveEditor/Private/Filters/PromotedFilterContainer.cpp b/Engine/Source/Editor/CurveEditor/Private/Filters/PromotedFilterContainer.cpp
new file mode 100644
--- /dev/null
+++ b/Engine/Source/Editor/CurveEditor/Private/Filters/PromotedFilterContainer.cpp
@@ -0,0 +1,130 @@
+﻿// Copyright Epic Games, Inc. All Rights Reserved.
+
+#include "Filters/PromotedFilterContainer.h"
+
+#include "Algo/NoneOf.h"
+#include "CurveEditorCommands.h"
+#include "Filters/CurveEditorFilterBase.h"
+#include "Framework/Commands/UICommandInfo.h"
+#include "Framework/MultiBox/MultiBoxBuilder.h"
+#include "UObject/Class.h"
+#include "Templates/SubclassOf.h"
+
+#define LOCTEXT_NAMESPACE "FPromotedFilterContainer"
+
+namespace UE::CurveEditor
+{
+FPromotedFilterContainer::FPromotedFilterContainer(const FName InContextName)
+	: CommandContext(MakeShared<FBindingContext>(
+		InContextName,
+		// This is what it will show up in Editor Preferences
+		LOCTEXT("CurveEditorFilters.Description", "Curve Editor Filters"),
+		// Bindings are not allowed to be the same between parent and child contexts
+		FCurveEditorCommands::Get().GetContextName(),
+		FAppStyle::GetAppStyleSetName())
+		)
+{}
+
+FPromotedFilterContainer::~FPromotedFilterContainer()
+{
+	FInputBindingManager::Get().RemoveContextByName(CommandContext->GetContextName());
+}
+
+void FPromotedFilterContainer::AppendToBuilder(FToolBarBuilder& InToolBarBuilder) const
+{
+	for (const FFilterData& FilterData : PromotedFilters)
+	{
+		InToolBarBuilder.AddToolBarButton(FilterData.Command);
+	}
+}
+
+void FPromotedFilterContainer::AppendToBuilder(FMenuBuilder& InMenuBuilder) const
+{
+	for (const FFilterData& FilterData : PromotedFilters)
+	{
+		InMenuBuilder.AddMenuEntry(FilterData.Command);
+	}
+}
+
+void FPromotedFilterContainer::AddInstance(UCurveEditorFilterBase& InFilter)
+{
+	const int32 Index = IndexOf(InFilter.GetClass());
+	const bool bIsAlreadyContained = PromotedFilters.IsValidIndex(Index);
+	if (bIsAlreadyContained)
+	{
+		return;
+	}
+
+	// Assuming all passed in InFilters are in the same package (e.g. GetTransientPackage()), GetFName should be unique.
+	const FName CommandName = InFilter.GetFName();
+	const bool bIsUniqueName = Algo::NoneOf(PromotedFilters, [CommandName](const FFilterData& FilterData)
+	{
+		return FilterData.FilterInstance.GetFName() == CommandName;
+	});
+	if (!ensure(bIsUniqueName))
+	{
+		return;
+	}
+
+	UClass* FilterClass = InFilter.GetClass();
+	TSharedPtr<FUICommandInfo> Command;
+	FUICommandInfo::MakeCommandInfo(CommandContext, Command,
+		CommandName, 
+		UCurveEditorFilterBase::GetLabel(FilterClass),
+		UCurveEditorFilterBase::GetDescription(FilterClass),
+		UCurveEditorFilterBase::GetIcon(FilterClass),
+		EUserInterfaceActionType::Button,
+		FInputChord()
+		);
+	PromotedFilters.Emplace(InFilter, Command.ToSharedRef());
+	
+	OnFilterAddedDelegate.Broadcast(InFilter, Command.ToSharedRef());
+}
+	
+void FPromotedFilterContainer::RemoveInstance(UCurveEditorFilterBase& InFilter)
+{
+	const int32 Index = IndexOf(InFilter.GetClass());
+	if (PromotedFilters.IsValidIndex(Index))
+	{
+		RemoveAtInternal(Index);
+	}
+}
+
+void FPromotedFilterContainer::AddReferencedObjects(FReferenceCollector& Collector)
+{
+	for (int32 Index = 0; Index < PromotedFilters.Num(); ++Index)
+	{
+		TObjectPtr<UCurveEditorFilterBase> Filter = &PromotedFilters[Index].FilterInstance;
+		Collector.AddReferencedObject(Filter);
+
+		// GC deleted the object?
+		if (!Filter)
+		{
+			// Preserver order, no swap.
+			RemoveAtInternal(Index);
+			--Index;
+		}
+	}
+}
+
+int32 FPromotedFilterContainer::IndexOf(const UClass* InFilterClass) const
+{
+	return InFilterClass
+		? PromotedFilters.IndexOfByPredicate([InFilterClass](const FFilterData& FilterData){ return FilterData.FilterInstance.GetClass() == InFilterClass; })
+		: INDEX_NONE;
+}
+
+void FPromotedFilterContainer::RemoveAtInternal(const int32 InIndex)
+{
+	checkSlow(PromotedFilters.IsValidIndex(InIndex));
+	
+	UCurveEditorFilterBase* Filter = &PromotedFilters[InIndex].FilterInstance;
+	const TSharedRef<FUICommandInfo> Command = PromotedFilters[InIndex].Command;
+	FUICommandInfo::UnregisterCommandInfo(CommandContext, Command);
+	PromotedFilters.RemoveAt(InIndex);
+		
+	OnFilterRemovedDelegate.Broadcast(*Filter, Command);
+}
+}
+
+#undef LOCTEXT_NAMESPACE
\ No newline at end of file
diff --git a/Engine/Source/Editor/CurveEditor/Private/Filters/SCurveEditorFilterPanel.cpp b/Engine/Source/Editor/CurveEditor/Private/Filters/SCurveEditorFilterPanel.cpp
--- a/Engine/Source/Editor/CurveEditor/Private/Filters/SCurveEditorFilterPanel.cpp
+++ b/Engine/Source/Editor/CurveEditor/Private/Filters/SCurveEditorFilterPanel.cpp
@@ -12,6 +12,7 @@
 #include "Delegates/Delegate.h"
 #include "DetailsViewArgs.h"
 #include "Filters/CurveEditorFilterBase.h"
+#include "Filters/FilterUtils.h"
 #include "Framework/Application/SlateApplication.h"
 #include "HAL/Platform.h"
 #include "HAL/PlatformCrt.h"
@@ -162,22 +163,14 @@ FReply SCurveEditorFilterPanel::OnApplyClicked()
 	TSharedPtr<FCurveEditor> CurveEditor = WeakCurveEditor.Pin();
 	if (CurveEditor && DetailsView->GetSelectedObjects().Num() > 0)
 	{
-		const TMap<FCurveModelID, FKeyHandleSet>& SelectedKeys = CurveEditor->GetSelection().GetAll();
+		const FText TransactionText = FText::Format(
+			NSLOCTEXT("CurveEditorFilterApply", "Filter Curves", "Filtered {0}|plural(one=Curve, other=Curves)"),
+			CurveEditor->GetSelection().Count()
+			);
+		const FScopedTransaction Transaction(TransactionText);
 		
-		FText TransactionText = FText::Format(NSLOCTEXT("CurveEditorFilterApply", "Filter Curves", "Filtered {0}|plural(one=Curve, other=Curves)"), CurveEditor->GetSelection().Count());
-		FScopedTransaction Transaction(TransactionText);
-
 		UCurveEditorFilterBase* Filter = Cast<UCurveEditorFilterBase>(DetailsView->GetSelectedObjects()[0]);
-		TMap<FCurveModelID, FKeyHandleSet> OutKeysToSelect;
-		Filter->ApplyFilter(CurveEditor.ToSharedRef(), SelectedKeys, OutKeysToSelect);
-
-		// Clear their selection and then set it to the keys the filter thinks you should have selected.
-		CurveEditor->GetSelection().Clear();
-
-		for (const TTuple<FCurveModelID, FKeyHandleSet>& OutSet : OutKeysToSelect)
-		{
-			CurveEditor->GetSelection().Add(OutSet.Key, ECurvePointType::Key, OutSet.Value.AsArray());
-		}
+		UE::CurveEditor::FilterUtils::ApplyFilter(CurveEditor.ToSharedRef(), *Filter);
 	}
 
 	return FReply::Handled();
diff --git a/Engine/Source/Editor/CurveEditor/Private/SCurveEditorPanel.cpp b/Engine/Source/Editor/CurveEditor/Private/SCurveEditorPanel.cpp
--- a/Engine/Source/Editor/CurveEditor/Private/SCurveEditorPanel.cpp
+++ b/Engine/Source/Editor/CurveEditor/Private/SCurveEditorPanel.cpp
@@ -17,6 +17,7 @@
 #include "Curves/KeyHandle.h"
 #include "Delegates/Delegate.h"
 #include "Editor.h"
+#include "Filters/FilterUtils.h"
 #include "Filters/CurveEditorBakeFilter.h"
 #include "Filters/CurveEditorReduceFilter.h"
 #include "Filters/SCurveEditorFilterPanel.h"
@@ -62,6 +63,8 @@
 #include "SGridLineSpacingList.h"
 #include "ScopedTransaction.h"
 #include "SlotBase.h"
+#include "Filters/CurveEditorEulerFilter.h"
+#include "Filters/PromotedFilterContainer.h"
 #include "Styling/AppStyle.h"
 #include "Styling/CoreStyle.h"
 #include "Styling/ISlateStyle.h"
@@ -156,6 +159,10 @@ void SCurveEditorPanel::Construct(const FArguments& InArgs, TSharedRef<FCurveEdi
 
 	CommandList = MakeShared<FUICommandList>();
 	CommandList->Append(InCurveEditor->GetCommands().ToSharedRef());
+	
+	ToolbarPromotedFilterBinder = MakeUnique<UE::CurveEditor::FPromotedFilterCommandBinder>(
+		GetCurveEditor()->GetToolbarPromotedFilters().ToSharedRef(), CommandList.ToSharedRef(), CurveEditor.ToSharedRef()
+		);
 
 	BindCommands();
 
@@ -462,7 +469,7 @@ void SCurveEditorPanel::BindCommands()
 
 	// User Implementable Filter  just defaults to Bake since we know it exists...
 	CommandList->MapAction(FCurveEditorCommands::Get().OpenUserImplementableFilterWindow, FExecuteAction::CreateSP(this, &SCurveEditorPanel::ShowCurveFilterUI, TSubclassOf<UCurveEditorFilterBase>(UCurveEditorBakeFilter::StaticClass())));
-
+	
 	// Axis Snapping
 	{
 		FUIAction SetSnappingNoneAction(FExecuteAction::CreateSP(this, &SCurveEditorPanel::SetAxisSnapping, ECurveEditorSnapAxis::CESA_None));
@@ -1198,7 +1205,12 @@ TSharedPtr<FExtender> SCurveEditorPanel::GetToolbarExtender()
 
 	struct Local
 	{
-		static void FillToolbar(FToolBarBuilder& ToolBarBuilder, TSharedRef<SCurveKeyDetailPanel> InKeyDetailsPanel, TSharedRef<SCurveEditorPanel> InEditorPanel)
+		static void FillToolbar(
+			FToolBarBuilder& ToolBarBuilder,
+			TSharedRef<SCurveKeyDetailPanel> InKeyDetailsPanel,
+			TSharedRef<SCurveEditorPanel> InEditorPanel,
+			TSharedRef<UE::CurveEditor::FPromotedFilterContainer> InToolbarPromotedFilters
+			)
 		{
 			ToolBarBuilder.BeginSection("View");
 			ToolBarBuilder.BeginStyleOverride("CurveEditorToolbar");
@@ -1371,6 +1383,7 @@ TSharedPtr<FExtender> SCurveEditorPanel::GetToolbarExtender()
 				ToolBarBuilder.AddToolBarButton(FCurveEditorCommands::Get().FlattenTangents);
 				ToolBarBuilder.AddToolBarButton(FCurveEditorCommands::Get().StraightenTangents);
 				ToolBarBuilder.AddToolBarButton(FCurveEditorCommands::Get().OpenUserImplementableFilterWindow);
+				InToolbarPromotedFilters->AppendToBuilder(ToolBarBuilder);
 			}
 			ToolBarBuilder.EndSection();
 			ToolBarBuilder.EndStyleOverride();
@@ -1382,7 +1395,9 @@ TSharedPtr<FExtender> SCurveEditorPanel::GetToolbarExtender()
 		"Asset",
 		EExtensionHook::After,
 		GetCommands(),
-		FToolBarExtensionDelegate::CreateStatic(&Local::FillToolbar, KeyDetailsView.ToSharedRef(), SharedThis(this))
+		FToolBarExtensionDelegate::CreateStatic(
+			&Local::FillToolbar, KeyDetailsView.ToSharedRef(), SharedThis(this), GetCurveEditor()->GetToolbarPromotedFilters().ToSharedRef()
+			)
 	);
 
 	return Extender;
diff --git a/Engine/Source/Editor/CurveEditor/Public/CurveEditor.h b/Engine/Source/Editor/CurveEditor/Public/CurveEditor.h
--- a/Engine/Source/Editor/CurveEditor/Public/CurveEditor.h
+++ b/Engine/Source/Editor/CurveEditor/Public/CurveEditor.h
@@ -584,6 +584,9 @@ public:
 	 */
 	void BindCommands();
 
+	/** @return The object managing which filters are promoted to the toolbar for quick access. */
+	TSharedPtr<UE::CurveEditor::FPromotedFilterContainer> GetToolbarPromotedFilters() const;
+
 public:
 
 	/** Suspend or resume broadcast of curve array changing  */
diff --git a/Engine/Source/Editor/CurveEditor/Public/Filters/CurveEditorEulerFilter.h b/Engine/Source/Editor/CurveEditor/Public/Filters/CurveEditorEulerFilter.h
--- a/Engine/Source/Editor/CurveEditor/Public/Filters/CurveEditorEulerFilter.h
+++ b/Engine/Source/Editor/CurveEditor/Public/Filters/CurveEditorEulerFilter.h
@@ -16,15 +16,19 @@ class UObject;
 struct FCurveModelID;
 struct FKeyHandleSet;
 
-UCLASS(DisplayName = "Euler")
+UCLASS(DisplayName = "Euler", meta = (
+	CurveEditorLabel = "Euler Filter",
+	CurveEditorDescription = "Uses the Euler filter to correct jumps in rotation caused by angle wrapping or gimbal lock. For example, instead of wrapping around with a large jump from 178° to -178°, it smoothly interpolates a 4° transition, ensuring natural and continuous movement.")
+	)
 class CURVEEDITOR_API UCurveEditorEulerFilter : public UCurveEditorFilterBase
 {
 	GENERATED_BODY()
 public:
 	UCurveEditorEulerFilter(){}
 protected:
-	virtual void ApplyFilter_Impl(TSharedRef<FCurveEditor> InCurveEditor, const TMap<FCurveModelID, FKeyHandleSet>& InKeysToOperateOn, TMap<FCurveModelID, FKeyHandleSet>& OutKeysToSelect) override;
 
-public:
-	
+	//~ Begin UCurveEditorFilterBase Interface
+	virtual void ApplyFilter_Impl(TSharedRef<FCurveEditor> InCurveEditor, const TMap<FCurveModelID, FKeyHandleSet>& InKeysToOperateOn, TMap<FCurveModelID, FKeyHandleSet>& OutKeysToSelect) override;
+	virtual bool CanApplyFilter_Impl(TSharedRef<FCurveEditor> InCurveEditor) override;
+	//~ End UCurveEditorFilterBase Interface
 };
\ No newline at end of file
diff --git a/Engine/Source/Editor/CurveEditor/Public/Filters/CurveEditorFilterBase.h b/Engine/Source/Editor/CurveEditor/Public/Filters/CurveEditorFilterBase.h
--- a/Engine/Source/Editor/CurveEditor/Public/Filters/CurveEditorFilterBase.h
+++ b/Engine/Source/Editor/CurveEditor/Public/Filters/CurveEditorFilterBase.h
@@ -9,6 +9,9 @@
 #include "CurveEditorFilterBase.generated.h"
 
 class FCurveEditor;
+class FText;
+struct FSlateIcon;
+template <typename T> class TSubclassOf;
 
 /**
 * An abstract base class for Curve Editor Filtering. If you inherit from this class your class will automatically
@@ -27,6 +30,16 @@ class CURVEEDITOR_API UCurveEditorFilterBase : public UObject
 
 public:
 
+	/** @return The value of the UCLASS(meta = (CurveEditorLabel = "value")) */
+	static FText GetLabel(const TSubclassOf<UCurveEditorFilterBase>& InClass);
+	/** @return The value of the UCLASS(meta = (CurveEditorDescription = "value")) */
+	static FText GetDescription(const TSubclassOf<UCurveEditorFilterBase>& InClass);
+	/**
+	 * @return The icon associated with this class. Icons can be overriden in editor styles (e.g. FStarshipEditorStyle) with "ClassIcon.ClassName".
+	 * @see FSlateIconFinder::FindIconForClass
+	 */
+	static FSlateIcon GetIcon(const TSubclassOf<UCurveEditorFilterBase>& InClass);
+
 	/**
 	 * Initialize this filter and any parameters it depends on
 	 */
diff --git a/Engine/Source/Editor/CurveEditor/Public/Filters/PromotedFilterCommandBinder.h b/Engine/Source/Editor/CurveEditor/Public/Filters/PromotedFilterCommandBinder.h
new file mode 100644
--- /dev/null
+++ b/Engine/Source/Editor/CurveEditor/Public/Filters/PromotedFilterCommandBinder.h
@@ -0,0 +1,55 @@
+﻿// Copyright Epic Games, Inc. All Rights Reserved.
+
+#pragma once
+
+#include "Filters/CurveEditorFilterBase.h"
+#include "Templates/UnrealTemplate.h"
+#include "Templates/SharedPointer.h"
+
+class FCurveEditor;
+class FUICommandInfo;
+class FUICommandList;
+
+namespace UE::CurveEditor
+{
+class FPromotedFilterContainer;
+
+/**
+ * Binds and unbinds commands that are created by FPromotedFilterContainer.
+ * When a filter is added to FPromotedFilterContainer, it creates a command to invoke that filter. This class is responsible to bind / unbind it.
+ */
+class CURVEEDITOR_API FPromotedFilterCommandBinder : public FNoncopyable
+{
+public:
+	
+	/**
+	 * @param InContainer Holds the promoted filter commands to bind to. The caller ensures that it outlives the constructed object. 
+	 * @param InCommandList The command list to bind commands to.
+	 * @param InCurveEditor The curve editor on which the filters will be applied.
+	 */
+	explicit FPromotedFilterCommandBinder(
+		const TSharedRef<FPromotedFilterContainer>& InContainer,
+		const TSharedRef<FUICommandList>& InCommandList,
+		const TSharedRef<FCurveEditor>& InCurveEditor
+		);
+	~FPromotedFilterCommandBinder();
+
+private:
+
+	/** The filter container this object binds to. */
+	const TWeakPtr<FPromotedFilterContainer> Container;
+	/** Commands are added to and removed from this list. */
+	const TWeakPtr<FUICommandList> CommandList;
+	/** Passed to filter as argument. */
+	const TWeakPtr<FCurveEditor> CurveEditor;
+	
+	void OnFilterAdded(UCurveEditorFilterBase& InFilter, const TSharedRef<FUICommandInfo>& InCommand) const;
+	void OnFilterRemoved(UCurveEditorFilterBase& InFilter, const TSharedRef<FUICommandInfo>& InCommand) const;
+
+	void MapAction(UCurveEditorFilterBase& InFilter, const TSharedRef<FUICommandInfo>& InCommand, FUICommandList& CommandListPin) const;
+	
+	void ApplyFilter(UCurveEditorFilterBase* InFilter, TWeakPtr<FUICommandInfo> Command) const;
+	bool CanApplyFilter(UCurveEditorFilterBase* InFilter, TWeakPtr<FUICommandInfo> Command) const;
+};
+}
+
diff --git a/Engine/Source/Editor/CurveEditor/Public/Filters/PromotedFilterContainer.h b/Engine/Source/Editor/CurveEditor/Public/Filters/PromotedFilterContainer.h
new file mode 100644
--- /dev/null
+++ b/Engine/Source/Editor/CurveEditor/Public/Filters/PromotedFilterContainer.h
@@ -0,0 +1,117 @@
+﻿// Copyright Epic Games, Inc. All Rights Reserved.
+
+#pragma once
+
+#include "Delegates/Delegate.h"
+#include "Framework/Commands/InputChord.h"
+#include "HAL/Platform.h"
+#include "Internationalization/Text.h"
+#include "Templates/UnrealTemplate.h"
+#include "Templates/SharedPointer.h"
+#include "Textures/SlateIcon.h"
+#include "UObject/GCObject.h"
+
+#include <type_traits>
+
+class FBindingContext;
+class FCurveEditor;
+class FMenuBuilder;
+class FToolBarBuilder;
+class FUICommandInfo;
+class FUICommandList;
+class UCurveEditorFilterBase;
+template <typename T> class TSubclassOf;
+
+namespace UE::CurveEditor
+{
+/**
+ * Holds the state for surfacing filters to the toolbar from the SCurveEditorFilterPanel, so the user can access & apply them more quickly.
+ *
+ * TODO UE-230269: For now, the curve editor only adds the Euler filter to this container. Once UE-230269 is addressed (allowing users to surface
+ * filters to the toolbar via UI), this class should be extended to persistently save the user-specified filter settings,  e.g. the FGaussianParams
+ * for the UCurveEditorGaussianFilter, etc.
+ */
+class CURVEEDITOR_API FPromotedFilterContainer
+	: public FGCObject
+	, public FNoncopyable
+{
+public:
+	
+	/**
+	 * @param InContextName The ID to give the underlying FBindingContext. Must be globally unique.
+	 */
+	explicit FPromotedFilterContainer(FName InContextName);
+	virtual ~FPromotedFilterContainer() override;
+
+	/** Appends the filters to InToolbarBuilder. */
+	void AppendToBuilder(FToolBarBuilder& InToolBarBuilder) const;
+	/** Appends the filters to InMenuBuilder. */
+	void AppendToBuilder(FMenuBuilder& InMenuBuilder) const;
+	
+	/**
+	 * Promotes a filter.
+	 * For simplicity, each class can only have one instance promoted (e.g. you cannot call AddInstance with two different UCurveEditorEulerFilter instances).
+	 */
+	void AddInstance(UCurveEditorFilterBase& InFilter);
+	
+	/** Removes a filter instance. You can pass in a CDO, too. */
+	void RemoveInstance(UCurveEditorFilterBase& InFilter);
+
+	template<typename TCallback> requires std::is_invocable_v<TCallback, UCurveEditorFilterBase&, const TSharedRef<FUICommandInfo>&>
+	void ForEachFilter(TCallback&& Callback);
+
+	DECLARE_MULTICAST_DELEGATE_TwoParams(FFilterDelegate, UCurveEditorFilterBase&, const TSharedRef<FUICommandInfo>&);
+	/** Broadcasts when a filter is added to this container. */
+	FFilterDelegate& OnFilterAdded() { return OnFilterAddedDelegate; }
+	/** Broadcasts when a filter is removed from this container. */
+	FFilterDelegate& OnFilterRemoved() { return OnFilterRemovedDelegate; }
+
+	//~ Begin FGCObject Interface
+	virtual void AddReferencedObjects(FReferenceCollector& Collector) override;
+	virtual FString GetReferencerName() const override { return TEXT("FPromotedFilterContainer"); }
+	//~ Begin FGCObject Interface
+
+private:
+
+	/** Needed to dynamically create FUICommandInfo. */
+	const TSharedRef<FBindingContext> CommandContext;
+
+	/** Broadcasts when a filter is added to this container. */
+	FFilterDelegate OnFilterAddedDelegate;
+	/** Broadcasts when a filter is removed from this container. */
+	FFilterDelegate OnFilterRemovedDelegate;
+
+	struct FFilterData
+	{
+		/** The instance to apply. This can be the CDO or an actual owned instance. */
+		UCurveEditorFilterBase& FilterInstance;
+
+		/**
+		 * Command info for this instance, which is used to the filter to menu / toolbar builders.
+		 * In the future, it also allows users to bind shortcuts to it.
+		 */
+		TSharedRef<FUICommandInfo> Command;
+
+		explicit FFilterData(UCurveEditorFilterBase& InFilter UE_LIFETIMEBOUND, TSharedRef<FUICommandInfo> InCommand)
+			: FilterInstance(InFilter)
+			, Command(InCommand)
+		{}
+	};
+
+	/** The filters that have been promoted. */
+	TArray<FFilterData> PromotedFilters;
+
+	int32 IndexOf(const UClass* InFilterClass) const;
+	void RemoveAtInternal(int32 InIndex);
+};
+
+template <typename TCallback> requires std::is_invocable_v<TCallback, UCurveEditorFilterBase&, const TSharedRef<FUICommandInfo>&>
+void FPromotedFilterContainer::ForEachFilter(TCallback&& Callback)
+{
+	for (const FFilterData& FilterData : PromotedFilters)
+	{
+		Callback(FilterData.FilterInstance, FilterData.Command);
+	}
+}
+}
+
diff --git a/Engine/Source/Editor/CurveEditor/Public/ICurveEditorModule.h b/Engine/Source/Editor/CurveEditor/Public/ICurveEditorModule.h
--- a/Engine/Source/Editor/CurveEditor/Public/ICurveEditorModule.h
+++ b/Engine/Source/Editor/CurveEditor/Public/ICurveEditorModule.h
@@ -7,6 +7,7 @@
 #include "Containers/ArrayView.h"
 #include "CurveEditorTypes.h"
 
+
 class FCurveEditor;
 class SCurveEditorView;
 class ICurveEditorExtension;
@@ -14,6 +15,7 @@ class ICurveEditorToolExtension;
 class ITimeSliderController;
 class FExtender;
 class FUICommandList;
+namespace UE::CurveEditor { class FPromotedFilterContainer; }
 
 /** A delegate which will create an extension for the Curve Editor. Used for adding new buttons and functionality to the editor as a whole. */
 DECLARE_DELEGATE_RetVal_OneParam(TSharedRef<ICurveEditorExtension>, FOnCreateCurveEditorExtension, TWeakPtr<FCurveEditor>);
@@ -39,6 +41,16 @@ struct FCurveEditorInitParams
 class ICurveEditorModule : public IModuleInterface
 {
 public:
+
+	static ICurveEditorModule& Get()
+	{
+		return FModuleManager::Get().GetModuleChecked<ICurveEditorModule>("CurveEditor");
+	}
+	static bool IsAvailable()
+	{
+		return FModuleManager::Get().IsModuleLoaded("CurveEditor");
+	}
+	
 	DECLARE_DELEGATE_RetVal_OneParam(TSharedRef<FExtender>, FCurveEditorMenuExtender, const TSharedRef<FUICommandList>);
 	virtual ~ICurveEditorModule() {}
 
@@ -68,4 +80,9 @@ public:
 	
 	virtual TArrayView<const FOnCreateCurveEditorExtension> GetEditorExtensions() const = 0;
 	virtual TArrayView<const FOnCreateCurveEditorToolExtension> GetToolExtensions() const = 0;
+
+	/**
+	 * @return The global setting of which filters should be promoted to the toolbar in the curve editor.
+	 */
+	virtual TSharedPtr<UE::CurveEditor::FPromotedFilterContainer> GetGlobalToolbarPromotedFilters() const = 0;
 };
\ No newline at end of file
diff --git a/Engine/Source/Editor/CurveEditor/Public/SCurveEditorPanel.h b/Engine/Source/Editor/CurveEditor/Public/SCurveEditorPanel.h
--- a/Engine/Source/Editor/CurveEditor/Public/SCurveEditorPanel.h
+++ b/Engine/Source/Editor/CurveEditor/Public/SCurveEditorPanel.h
@@ -12,6 +12,7 @@
 #include "WidgetFocusUtils.h"
 #include "Curves/RealCurve.h"
 #include "Curves/RichCurve.h"
+#include "Filters/PromotedFilterCommandBinder.h"
 #include "HAL/Platform.h"
 #include "Input/Reply.h"
 #include "Internationalization/Text.h"
@@ -47,6 +48,7 @@ struct FCurveEditorDelayedDrag;
 struct FCurveEditorEditObjectContainer;
 struct FCurveEditorToolID;
 struct FKeyEvent;
+namespace UE::CurveEditor { class FPromotedFilterContainer; }
 
 /**
  * Curve editor widget that reflects the state of an FCurveEditor
@@ -356,6 +358,9 @@ private:
 	/** Command list for widget specific command bindings */
 	TSharedPtr<FUICommandList> CommandList;
 
+	/** Binds commands from FPromotedFilterContainer so the filters surfaced from SCurveEditorFilterPanel show up in the toolbar. */
+	TUniquePtr<UE::CurveEditor::FPromotedFilterCommandBinder> ToolbarPromotedFilterBinder;
+
 	/** Cached serial number from the curve editor selection. Used to update edit UIs when the selection changes. */
 	uint32 CachedSelectionSerialNumber;
 
diff --git a/Engine/Source/Editor/EditorStyle/Private/StarshipStyle.cpp b/Engine/Source/Editor/EditorStyle/Private/StarshipStyle.cpp
--- a/Engine/Source/Editor/EditorStyle/Private/StarshipStyle.cpp
+++ b/Engine/Source/Editor/EditorStyle/Private/StarshipStyle.cpp
@@ -6279,6 +6279,9 @@ void FStarshipEditorStyle::FStyle::SetupClassIconsAndThumbnails()
 		Set("ClassIcon.StationaryMobilityIcon", new IMAGE_BRUSH("Icons/ActorIcons/Light_Adjustable_16x", Icon16x16));
 		Set("ClassIcon.ComponentMobilityHeaderIcon", new IMAGE_BRUSH("Icons/ActorIcons/ComponentMobilityHeader_7x16", Icon7x16));
 
+		// Curve Editor icons
+		Set("ClassIcon.CurveEditorEulerFilter", new IMAGE_BRUSH_SVG("Starship/GenericCurveEditor/CurveEditorOpEuler", Icon16x16));
+
 		// Asset Type Classes
 		const TCHAR* AssetTypes[] = {
 
