Curve Editor: Improve mouse inputs for editing tangents

* Restore middle mouse button can edit tangents
* MMB clicking an existing handle or key no longer creates a new key
* Shift-LMB or MMb drag adjusts the weight of weighted tangents
* Ctrl-Shift-LMB or MMB drag adjusts the tangent of weighted tangents

diff --git a/Engine/Source/Editor/CurveEditor/Private/DragOperations/CurveEditorDragOperation_Tangent.cpp b/Engine/Source/Editor/CurveEditor/Private/DragOperations/CurveEditorDragOperation_Tangent.cpp
--- a/Engine/Source/Editor/CurveEditor/Private/DragOperations/CurveEditorDragOperation_Tangent.cpp
+++ b/Engine/Source/Editor/CurveEditor/Private/DragOperations/CurveEditorDragOperation_Tangent.cpp
@@ -105,7 +105,8 @@ namespace CurveEditorDragOperation
 
 	void FTangentSolver::SolveTangentConstantWeight(const FSolverParams& Params)
 	{
-		if (Params.MouseEvent.IsMouseButtonDown(EKeys::LeftMouseButton))
+		if (Params.MouseEvent.IsMouseButtonDown(EKeys::LeftMouseButton) ||
+			Params.MouseEvent.IsMouseButtonDown(EKeys::MiddleMouseButton))
 		{				
 			const float PixelLength = Params.bArriveTangent ? -60.0f : 60.f;
 			const float Slope = Params.InitialTangent * -Params.DisplayRatio;
@@ -140,17 +141,29 @@ namespace CurveEditorDragOperation
 			InitialWeight);
 
 		FVector2D NewTangentOffset;
-		bool bValueSet = false;
 
-		if (Params.MouseEvent.IsMouseButtonDown(EKeys::LeftMouseButton))
+		if (Params.MouseEvent.IsMouseButtonDown(EKeys::LeftMouseButton) ||
+			Params.MouseEvent.IsMouseButtonDown(EKeys::MiddleMouseButton))
 		{
 			const FVector2D PreviousTangentOffset = LastTangentOffset.IsSet() ? LastTangentOffset.GetValue() : InitialTangentOffset;
 			const FVector2D PointAtOffset = InitialTangentOffset + Params.PixelDelta;
 
 			if (Params.MouseEvent.IsShiftDown() &&
 				Params.MouseEvent.IsControlDown())
-			{
-				// Shift-LMB adjusts the weight, but not the tangent
+			{				
+				// Draging with Ctrl-Shift modifiers adjusts the tangent, but not the weight
+				const double Weigth = PreviousTangentOffset.Size();
+				NewTangentOffset = PointAtOffset.GetSafeNormal() * Weigth;
+
+				if ((Params.bArriveTangent && NewTangentOffset.X >= 0.0) ||
+					(!Params.bArriveTangent && NewTangentOffset.X <= 0.0))
+				{
+					NewTangentOffset = { 0.0, Weigth * FMath::Sign(NewTangentOffset.Y) };
+				}
+			}
+			else if (Params.MouseEvent.IsShiftDown())
+			{				
+				// Draging with Shift modifier adjusts the weight, but not the tangent
 				const double Scalar = FVector2D::DotProduct(PointAtOffset, PreviousTangentOffset.GetSafeNormal());
 
 				if (Scalar >= 1.0)
@@ -162,21 +175,9 @@ namespace CurveEditorDragOperation
 					NewTangentOffset = PreviousTangentOffset.GetSafeNormal();
 				}
 			}
-			else if (Params.MouseEvent.IsShiftDown())
-			{
-				// Ctrl-Shift-LMB adjusts the tangent, but not the weight
-				const double Weigth = PreviousTangentOffset.Size();
-				NewTangentOffset = PointAtOffset.GetSafeNormal() * Weigth;
-
-				if ((Params.bArriveTangent && NewTangentOffset.X >= 0.0) ||
-					(!Params.bArriveTangent && NewTangentOffset.X <= 0.0))
-				{
-					NewTangentOffset = { 0.0, Weigth * FMath::Sign(NewTangentOffset.Y) };
-				}
-			}
 			else
 			{
-				// LMB adjusts the tangent and weight 
+				// Draging without modifiers adjusts the tangent and weight 
 				NewTangentOffset = InitialTangentOffset + Params.PixelDelta;
 			}
 
