Curve Editor: Keep value when deleting all keys. Implemented for ChannelCurve and RichCurve that support default values.

diff --git a/Engine/Plugins/Animation/ControlRig/Source/ControlRigEditor/Private/ControlRigSpaceChannelCurveModel.cpp b/Engine/Plugins/Animation/ControlRig/Source/ControlRigEditor/Private/ControlRigSpaceChannelCurveModel.cpp
--- a/Engine/Plugins/Animation/ControlRig/Source/ControlRigEditor/Private/ControlRigSpaceChannelCurveModel.cpp
+++ b/Engine/Plugins/Animation/ControlRig/Source/ControlRigEditor/Private/ControlRigSpaceChannelCurveModel.cpp
@@ -135,7 +135,7 @@ bool FControlRigSpaceChannelCurveModel::Evaluate(double Time, double& OutValue)
 	return false;
 }
 
-void FControlRigSpaceChannelCurveModel::RemoveKeys(TArrayView<const FKeyHandle> InKeys)
+void FControlRigSpaceChannelCurveModel::RemoveKeys(TArrayView<const FKeyHandle> InKeys, double InCurrentTime)
 {
 	FMovieSceneControlRigSpaceChannel* Channel = ChannelHandle.Get();
 	UMovieSceneSection*      Section = WeakSection.Get();
diff --git a/Engine/Plugins/Animation/ControlRig/Source/ControlRigEditor/Public/ControlRigSpaceChannelCurveModel.h b/Engine/Plugins/Animation/ControlRig/Source/ControlRigEditor/Public/ControlRigSpaceChannelCurveModel.h
--- a/Engine/Plugins/Animation/ControlRig/Source/ControlRigEditor/Public/ControlRigSpaceChannelCurveModel.h
+++ b/Engine/Plugins/Animation/ControlRig/Source/ControlRigEditor/Public/ControlRigSpaceChannelCurveModel.h
@@ -43,7 +43,7 @@ public:
 	virtual void GetNeighboringKeys(const FKeyHandle InKeyHandle, TOptional<FKeyHandle>& OutPreviousKeyHandle, TOptional<FKeyHandle>& OutNextKeyHandle) const override {}
 	virtual bool Evaluate(double ProspectiveTime, double& OutValue) const override;
 	virtual void AddKeys(TArrayView<const FKeyPosition> InKeyPositions, TArrayView<const FKeyAttributes> InAttributes, TArrayView<TOptional<FKeyHandle>>* OutKeyHandles) override;
-	virtual void RemoveKeys(TArrayView<const FKeyHandle> InKeys) override;
+	virtual void RemoveKeys(TArrayView<const FKeyHandle> InKeys, double InCurrentTime) override;
 
 	virtual void CreateKeyProxies(TArrayView<const FKeyHandle> InKeyHandles, TArrayView<UObject*> OutObjects) override;
 	
diff --git a/Engine/Plugins/Editor/CurveEditorTools/Source/CurveEditorTools/Private/Filters/CurveEditorFFTFilter.cpp b/Engine/Plugins/Editor/CurveEditorTools/Source/CurveEditorTools/Private/Filters/CurveEditorFFTFilter.cpp
--- a/Engine/Plugins/Editor/CurveEditorTools/Source/CurveEditorTools/Private/Filters/CurveEditorFFTFilter.cpp
+++ b/Engine/Plugins/Editor/CurveEditorTools/Source/CurveEditorTools/Private/Filters/CurveEditorFFTFilter.cpp
@@ -161,7 +161,7 @@ void UCurveEditorFFTFilter::ApplyFilter_Impl(TSharedRef<FCurveEditor> InCurveEdi
 		}
 
 		// Our curve model is still full of baked data. We do this as a second loop otherwise as we added keys, the sampling evaluation would be incorrect.
-		CurveModel->RemoveKeys(BakedKeyHandles);
+		CurveModel->RemoveKeys(BakedKeyHandles, 0.0);
 
 		FKeyHandleSet& SelectedHandleSet = OutKeysToSelect.Add(Pair.Key);
 		for(int32 KeyIndex = 0; KeyIndex < OriginalKeyPositions.Num(); KeyIndex++)
diff --git a/Engine/Plugins/Editor/CurveEditorTools/Source/CurveEditorTools/Private/Tools/CurveEditorTransformTool.cpp b/Engine/Plugins/Editor/CurveEditorTools/Source/CurveEditorTools/Private/Tools/CurveEditorTransformTool.cpp
--- a/Engine/Plugins/Editor/CurveEditorTools/Source/CurveEditorTools/Private/Tools/CurveEditorTransformTool.cpp
+++ b/Engine/Plugins/Editor/CurveEditorTools/Source/CurveEditorTools/Private/Tools/CurveEditorTransformTool.cpp
@@ -1128,7 +1128,7 @@ void FCurveEditorTransformTool::OnDragEnd()
 			}
 
 			// Remove any keys that overlap before moving new keys on top
-			Curve->RemoveKeys(KeysToRemove);
+			Curve->RemoveKeys(KeysToRemove, 0.0);
 
 			// Then, move the keys
 			Curve->SetKeyPositions(KeyData.Handles, KeyData.LastDraggedKeyPositions, EPropertyChangeType::ValueSet);
diff --git a/Engine/Plugins/FX/Niagara/Source/NiagaraEditorWidgets/Private/DetailCustomizations/NiagaraDataInterfaceCurveDetails.cpp b/Engine/Plugins/FX/Niagara/Source/NiagaraEditorWidgets/Private/DetailCustomizations/NiagaraDataInterfaceCurveDetails.cpp
--- a/Engine/Plugins/FX/Niagara/Source/NiagaraEditorWidgets/Private/DetailCustomizations/NiagaraDataInterfaceCurveDetails.cpp
+++ b/Engine/Plugins/FX/Niagara/Source/NiagaraEditorWidgets/Private/DetailCustomizations/NiagaraDataInterfaceCurveDetails.cpp
@@ -1550,7 +1550,7 @@ FReply FNiagaraDataInterfaceCurveDetailsBase::CurveTemplateSelected(TWeakObjectP
 				{
 					TArray<FKeyHandle> KeyHandles;
 					CurveModel->GetKeys(TNumericLimits<double>::Lowest(), TNumericLimits<double>::Max(), TNumericLimits<double>::Lowest(), TNumericLimits<double>::Max(), KeyHandles);
-					CurveModel->RemoveKeys(KeyHandles);
+					CurveModel->RemoveKeys(KeyHandles,0.0);
 
 					const FRichCurve& FloatCurve = FloatCurveAsset->FloatCurve;
 					for (auto KeyIterator = FloatCurve.GetKeyHandleIterator(); KeyIterator; ++KeyIterator)
diff --git a/Engine/Plugins/Mutable/Source/CustomizableObjectEditor/Private/MuCOE/SMutableCurveViewer.h b/Engine/Plugins/Mutable/Source/CustomizableObjectEditor/Private/MuCOE/SMutableCurveViewer.h
--- a/Engine/Plugins/Mutable/Source/CustomizableObjectEditor/Private/MuCOE/SMutableCurveViewer.h
+++ b/Engine/Plugins/Mutable/Source/CustomizableObjectEditor/Private/MuCOE/SMutableCurveViewer.h
@@ -55,7 +55,7 @@ public:
 	virtual void AddKeys(TArrayView<const FKeyPosition> InPositions, TArrayView<const FKeyAttributes> InAttributes,
 		TArrayView<TOptional<FKeyHandle>>* OutKeyHandles) override;
 
-	virtual void RemoveKeys(TArrayView<const FKeyHandle> InKeys) override {}
+	virtual void RemoveKeys(TArrayView<const FKeyHandle> InKeys, double InCurrentTime) override {}
 
 	virtual void GetKeyPositions(TArrayView<const FKeyHandle> InKeys,
 		TArrayView<FKeyPosition> OutKeyPositions) const override;
diff --git a/Engine/Plugins/VirtualProduction/CameraCalibration/Source/CameraCalibrationEditor/Private/AssetEditor/Curves/LensDataCurveModel.cpp b/Engine/Plugins/VirtualProduction/CameraCalibration/Source/CameraCalibrationEditor/Private/AssetEditor/Curves/LensDataCurveModel.cpp
--- a/Engine/Plugins/VirtualProduction/CameraCalibration/Source/CameraCalibrationEditor/Private/AssetEditor/Curves/LensDataCurveModel.cpp
+++ b/Engine/Plugins/VirtualProduction/CameraCalibration/Source/CameraCalibrationEditor/Private/AssetEditor/Curves/LensDataCurveModel.cpp
@@ -23,7 +23,7 @@ void FLensDataCurveModel::AddKeys(TArrayView<const FKeyPosition> InKeyPositions,
 	//Don't support adding keys from curve editor by default. Specific models can override
 }
 
-void FLensDataCurveModel::RemoveKeys(TArrayView<const FKeyHandle> InKeys)
+void FLensDataCurveModel::RemoveKeys(TArrayView<const FKeyHandle> InKeys, double InCurrentTime)
 {
 	//Don't support removing keys from curve editor by default. Specific models can override
 }
diff --git a/Engine/Plugins/VirtualProduction/CameraCalibration/Source/CameraCalibrationEditor/Private/AssetEditor/Curves/LensDataCurveModel.h b/Engine/Plugins/VirtualProduction/CameraCalibration/Source/CameraCalibrationEditor/Private/AssetEditor/Curves/LensDataCurveModel.h
--- a/Engine/Plugins/VirtualProduction/CameraCalibration/Source/CameraCalibrationEditor/Private/AssetEditor/Curves/LensDataCurveModel.h
+++ b/Engine/Plugins/VirtualProduction/CameraCalibration/Source/CameraCalibrationEditor/Private/AssetEditor/Curves/LensDataCurveModel.h
@@ -21,7 +21,7 @@ public:
 
 	//~ Begin FRichCurveEditorModel
 	virtual void AddKeys(TArrayView<const FKeyPosition> InKeyPositions, TArrayView<const FKeyAttributes> InAttributes, TArrayView<TOptional<FKeyHandle>>* OutKeyHandles) override;
-	virtual void RemoveKeys(TArrayView<const FKeyHandle> InKeys) override;
+	virtual void RemoveKeys(TArrayView<const FKeyHandle> InKeys, double InCurrentTime) override;
 	virtual void SetKeyPositions(TArrayView<const FKeyHandle> InKeys, TArrayView<const FKeyPosition> InKeyPositions, EPropertyChangeType::Type ChangeType) override;
 	virtual bool IsValid() const override;
 	virtual FRichCurve& GetRichCurve() override;
diff --git a/Engine/Source/Editor/CurveEditor/Private/CurveEditor.cpp b/Engine/Source/Editor/CurveEditor/Private/CurveEditor.cpp
--- a/Engine/Source/Editor/CurveEditor/Private/CurveEditor.cpp
+++ b/Engine/Source/Editor/CurveEditor/Private/CurveEditor.cpp
@@ -1645,6 +1645,12 @@ bool FCurveEditor::CopyBufferCurveToCurveID(const UCurveEditorCopyableCurveKeys*
 	// because mixing the pasted results with the original results wouldn't make any sense.
 	if (bInOverwriteRange)
 	{
+		double CurrentTime = 0.0;
+		if (TSharedPtr<ITimeSliderController> TimeSliderController = WeakTimeSliderController.Pin())
+		{
+			FFrameRate TickResolution = TimeSliderController->GetTickResolution();
+			CurrentTime = TickResolution.AsSeconds(TimeSliderController->GetScrubPosition());
+		}
 		TArray<FKeyHandle> KeysToRemove;
 		double MinKeyTime = TNumericLimits<double>::Max();
 		double MaxKeyTime = TNumericLimits<double>::Lowest(); 
@@ -1671,7 +1677,7 @@ bool FCurveEditor::CopyBufferCurveToCurveID(const UCurveEditorCopyableCurveKeys*
 			TargetCurve->GetKeys(MinKeyTime, MaxKeyTime, TNumericLimits<double>::Lowest(), TNumericLimits<double>::Max(), KeysToRemove);
 		}
 
-		TargetCurve->RemoveKeys(KeysToRemove);
+		TargetCurve->RemoveKeys(KeysToRemove, CurrentTime);
 	}
 
 	for (int32 Index = 0; Index < InSourceCurve->KeyPositions.Num(); ++Index)
@@ -1890,12 +1896,19 @@ void FCurveEditor::DeleteSelection()
 {
 	FScopedTransaction Transaction(LOCTEXT("DeleteKeys", "Delete Keys"));
 
+	double CurrentTime = 0.0;
+	if (TSharedPtr<ITimeSliderController> TimeSliderController = WeakTimeSliderController.Pin())
+	{
+		FFrameRate TickResolution = TimeSliderController->GetTickResolution();
+		CurrentTime = TickResolution.AsSeconds(TimeSliderController->GetScrubPosition());
+	}
+
 	for (const TTuple<FCurveModelID, FKeyHandleSet>& Pair : Selection.GetAll())
 	{
 		if (FCurveModel* Curve = FindCurve(Pair.Key))
 		{
 			Curve->Modify();
-			Curve->RemoveKeys(Pair.Value.AsArray());
+			Curve->RemoveKeys(Pair.Value.AsArray(), CurrentTime);
 		}
 	}
 
@@ -2213,8 +2226,14 @@ void FCurveEditor::ApplyBufferedCurveToTarget(const IBufferedCurveModel* Buffere
 	TArray<FKeyHandle> TargetKeyHandles;
 	TargetCurve->GetKeys(TNumericLimits<double>::Lowest(), TNumericLimits<double>::Max(), TNumericLimits<double>::Lowest(), TNumericLimits<double>::Max(), TargetKeyHandles);
 
+	double CurrentTime = 0.0;
+	if (TSharedPtr<ITimeSliderController> TimeSliderController = WeakTimeSliderController.Pin())
+	{
+		FFrameRate TickResolution = TimeSliderController->GetTickResolution();
+		CurrentTime = TickResolution.AsSeconds(TimeSliderController->GetScrubPosition());
+	}
 	// Clear our current keys from the target curve
-	TargetCurve->RemoveKeys(TargetKeyHandles);
+	TargetCurve->RemoveKeys(TargetKeyHandles, CurrentTime);
 
 	// Now put our buffered keys into the target curve
 	TargetCurve->AddKeys(KeyPositions, KeyAttributes);
diff --git a/Engine/Source/Editor/CurveEditor/Private/DragOperations/CurveEditorDragOperation_MoveKeys.cpp b/Engine/Source/Editor/CurveEditor/Private/DragOperations/CurveEditorDragOperation_MoveKeys.cpp
--- a/Engine/Source/Editor/CurveEditor/Private/DragOperations/CurveEditorDragOperation_MoveKeys.cpp
+++ b/Engine/Source/Editor/CurveEditor/Private/DragOperations/CurveEditorDragOperation_MoveKeys.cpp
@@ -158,7 +158,7 @@ void FCurveEditorDragOperation_MoveKeys::OnEndDrag(FVector2D InitialPosition, FV
 			}
 
 			// Remove any keys that overlap before moving new keys on top
-			Curve->RemoveKeys(KeysToRemove);
+			Curve->RemoveKeys(KeysToRemove, 0.0);
 			// Then, move the keys
 			Curve->SetKeyPositions(KeyData.Handles, KeyData.LastDraggedKeyPositions, EPropertyChangeType::ValueSet);
 		}
diff --git a/Engine/Source/Editor/CurveEditor/Private/Filters/CurveEditorBakeFilter.cpp b/Engine/Source/Editor/CurveEditor/Private/Filters/CurveEditorBakeFilter.cpp
--- a/Engine/Source/Editor/CurveEditor/Private/Filters/CurveEditorBakeFilter.cpp
+++ b/Engine/Source/Editor/CurveEditor/Private/Filters/CurveEditorBakeFilter.cpp
@@ -239,7 +239,7 @@ void UCurveEditorBakeFilter::ApplyFilter_Impl(TSharedRef<FCurveEditor> InCurveEd
 
 
 		// Remove all the old in-between keys and add the new ones
-		Curve->RemoveKeys(KeyHandles);
+		Curve->RemoveKeys(KeyHandles,0.0);
 		Curve->AddKeys(NewKeyPositions, NewKeyAttributes, &NewKeyHandlesView);
 
 		for (const TOptional<FKeyHandle>& Handle : NewKeyHandlesView)
diff --git a/Engine/Source/Editor/CurveEditor/Private/Filters/CurveEditorReduceFilter.cpp b/Engine/Source/Editor/CurveEditor/Private/Filters/CurveEditorReduceFilter.cpp
--- a/Engine/Source/Editor/CurveEditor/Private/Filters/CurveEditorReduceFilter.cpp
+++ b/Engine/Source/Editor/CurveEditor/Private/Filters/CurveEditorReduceFilter.cpp
@@ -179,7 +179,7 @@ void UCurveEditorReduceFilter::ApplyFilter_Impl(TSharedRef<FCurveEditor> InCurve
 				}
 			}
 			Curve->Modify();
-			Curve->RemoveKeys(KeysToRemove);
+			Curve->RemoveKeys(KeysToRemove,0.0);
 		}
 	}
 }
diff --git a/Engine/Source/Editor/CurveEditor/Private/Filters/CurveEditorSmartReduceFilter.cpp b/Engine/Source/Editor/CurveEditor/Private/Filters/CurveEditorSmartReduceFilter.cpp
--- a/Engine/Source/Editor/CurveEditor/Private/Filters/CurveEditorSmartReduceFilter.cpp
+++ b/Engine/Source/Editor/CurveEditor/Private/Filters/CurveEditorSmartReduceFilter.cpp
@@ -210,7 +210,7 @@ void UCurveEditorSmartReduceFilter::SmartReduce(FCurveModel* Curve, const FSmart
 	TArray<double> Velocities;
 	Difference(Cache.Positions, Velocities);
 	Curve->Modify();
-	Curve->RemoveKeys(KeyHandles);
+	Curve->RemoveKeys(KeyHandles, 0.0);
 	FKeyAttributes NewDefaultAttribute;
 	NewDefaultAttribute.SetInterpMode(ERichCurveInterpMode::RCIM_Cubic);
 	NewDefaultAttribute.SetTangentMode(ERichCurveTangentMode::RCTM_SmartAuto);
diff --git a/Engine/Source/Editor/CurveEditor/Private/RichCurveEditorModel.cpp b/Engine/Source/Editor/CurveEditor/Private/RichCurveEditorModel.cpp
--- a/Engine/Source/Editor/CurveEditor/Private/RichCurveEditorModel.cpp
+++ b/Engine/Source/Editor/CurveEditor/Private/RichCurveEditorModel.cpp
@@ -278,7 +278,7 @@ bool FRichCurveEditorModel::Evaluate(double Time, double& OutValue) const
 	return false;
 }
 
-void FRichCurveEditorModel::RemoveKeys(TArrayView<const FKeyHandle> InKeys)
+void FRichCurveEditorModel::RemoveKeys(TArrayView<const FKeyHandle> InKeys, double InCurrentTime)
 {
 	if (UObject* Owner = WeakOwner.Get())
 	{
@@ -286,6 +286,10 @@ void FRichCurveEditorModel::RemoveKeys(TArrayView<const FKeyHandle> InKeys)
 		{
 			Owner->Modify();
 			FRichCurve& RichCurve = GetRichCurve();
+
+			double CurrentValue = 0.0;
+			Evaluate(InCurrentTime, CurrentValue);
+
 			for (FKeyHandle Handle : InKeys)
 			{
 				if (RichCurve.IsKeyHandleValid(Handle))
@@ -293,6 +297,10 @@ void FRichCurveEditorModel::RemoveKeys(TArrayView<const FKeyHandle> InKeys)
 					RichCurve.DeleteKey(Handle);
 				}
 			}
+			if (RichCurve.GetNumKeys() == 0)
+			{
+				RichCurve.DefaultValue = CurrentValue;
+			}
 
 			CurveModifiedDelegate.Broadcast();
 		}
diff --git a/Engine/Source/Editor/CurveEditor/Public/CurveModel.h b/Engine/Source/Editor/CurveEditor/Public/CurveModel.h
--- a/Engine/Source/Editor/CurveEditor/Public/CurveModel.h
+++ b/Engine/Source/Editor/CurveEditor/Public/CurveModel.h
@@ -100,12 +100,18 @@ public:
 	 */
 	virtual void AddKeys(TArrayView<const FKeyPosition> InPositions, TArrayView<const FKeyAttributes> InAttributes, TArrayView<TOptional<FKeyHandle>>* OutKeyHandles = nullptr) = 0;
 
+	UE_DEPRECATED(5.6, "Implement the version that uses the current time.")
+	virtual void RemoveKeys(TArrayView<const FKeyHandle> InKeys)
+	{
+		RemoveKeys(InKeys, 0.0);
+	}
 	/**
 	 * Remove all the keys with the specified key handles from this curve
 	 *
-	 * @param InKeys                  Array of key handles to be removed from this curve
+	 * @param InKeys				Array of key handles to be removed from this curve
+	 * @param InCurrentTime			Current time slider position, may be used to set any default value to the current value at that time if no keys exist
 	 */
-	virtual void RemoveKeys(TArrayView<const FKeyHandle> InKeys) = 0;
+	virtual void RemoveKeys(TArrayView<const FKeyHandle> InKeys, double InCurrentTime) = 0;
 
 	/**
 	 * Retrieve all key positions that pertain to the specified input key handles
diff --git a/Engine/Source/Editor/CurveEditor/Public/KeyBarCurveModel.h b/Engine/Source/Editor/CurveEditor/Public/KeyBarCurveModel.h
--- a/Engine/Source/Editor/CurveEditor/Public/KeyBarCurveModel.h
+++ b/Engine/Source/Editor/CurveEditor/Public/KeyBarCurveModel.h
@@ -38,7 +38,7 @@ public:
 	virtual void DrawCurve(const FCurveEditor& CurveEditor, const FCurveEditorScreenSpace& ScreenSpace, TArray<TTuple<double, double>>& InterpolatingPoints) const override{};
 	virtual void GetKeys(double MinTime, double MaxTime, double MinValue, double MaxValue, TArray<FKeyHandle>& OutKeyHandles) const override {};
 	virtual void AddKeys(TArrayView<const FKeyPosition> InPositions, TArrayView<const FKeyAttributes> InAttributes, TArrayView<TOptional<FKeyHandle>>* OutKeyHandles = nullptr) override {};
-	virtual void RemoveKeys(TArrayView<const FKeyHandle> InKeys) override {};
+	virtual void RemoveKeys(TArrayView<const FKeyHandle> InKeys, double InCurrentTime) override {};
 	virtual void GetKeyPositions(TArrayView<const FKeyHandle> InKeys, TArrayView<FKeyPosition> OutKeyPositions) const override {};
 	virtual void SetKeyPositions(TArrayView<const FKeyHandle> InKeys, TArrayView<const FKeyPosition> InKeyPositions,const EPropertyChangeType::Type ChangeType = EPropertyChangeType::Unspecified) override {};
 	virtual void GetKeyDrawInfo(ECurvePointType PointType, const FKeyHandle InKeyHandle, FKeyDrawInfo& OutDrawInfo) const override {};
diff --git a/Engine/Source/Editor/CurveEditor/Public/RichCurveEditorModel.h b/Engine/Source/Editor/CurveEditor/Public/RichCurveEditorModel.h
--- a/Engine/Source/Editor/CurveEditor/Public/RichCurveEditorModel.h
+++ b/Engine/Source/Editor/CurveEditor/Public/RichCurveEditorModel.h
@@ -58,7 +58,7 @@ public:
 
 	virtual bool Evaluate(double ProspectiveTime, double& OutValue) const override;
 	virtual void AddKeys(TArrayView<const FKeyPosition> InKeyPositions, TArrayView<const FKeyAttributes> InAttributes, TArrayView<TOptional<FKeyHandle>>* OutKeyHandles) override;
-	virtual void RemoveKeys(TArrayView<const FKeyHandle> InKeys) override;
+	virtual void RemoveKeys(TArrayView<const FKeyHandle> InKeys, double InCurrentTime) override;
 
 	virtual void CreateKeyProxies(TArrayView<const FKeyHandle> InKeyHandles, TArrayView<UObject*> OutObjects) override;
 
diff --git a/Engine/Source/Editor/CurveTableEditor/Private/RealCurveModel.cpp b/Engine/Source/Editor/CurveTableEditor/Private/RealCurveModel.cpp
--- a/Engine/Source/Editor/CurveTableEditor/Private/RealCurveModel.cpp
+++ b/Engine/Source/Editor/CurveTableEditor/Private/RealCurveModel.cpp
@@ -103,7 +103,7 @@ void FRealCurveModel::AddKeys(TArrayView<const FKeyPosition> InKeyPositions, TAr
 	// Left empty for now - The curve editor cannot directly add/remove keys to RealCurves
 }
 
-void FRealCurveModel::RemoveKeys(TArrayView<const FKeyHandle> InKeys)
+void FRealCurveModel::RemoveKeys(TArrayView<const FKeyHandle> InKeys, double InCurrentTime)
 {
 	// Left empty for now - The curve editor cannot directly add/remove keys to RealCurves
 }
diff --git a/Engine/Source/Editor/CurveTableEditor/Private/RealCurveModel.h b/Engine/Source/Editor/CurveTableEditor/Private/RealCurveModel.h
--- a/Engine/Source/Editor/CurveTableEditor/Private/RealCurveModel.h
+++ b/Engine/Source/Editor/CurveTableEditor/Private/RealCurveModel.h
@@ -40,7 +40,7 @@ class FRealCurveModel : public FCurveModel
 	virtual void DrawCurve(const FCurveEditor& CurveEditor, const FCurveEditorScreenSpace& ScreenSpace, TArray<TTuple<double, double>>& InterpolatingPoints) const override;
 	virtual void GetKeys(double MinTime, double MaxTime, double MinValue, double MaxValue, TArray<FKeyHandle>& OutKeyHandles) const override;
 	virtual void AddKeys(TArrayView<const FKeyPosition> InKeyPositions, TArrayView<const FKeyAttributes> InAttributes, TArrayView<TOptional<FKeyHandle>>* OutKeyHandles) override;
-	virtual void RemoveKeys(TArrayView<const FKeyHandle> InKeys) override;
+	virtual void RemoveKeys(TArrayView<const FKeyHandle> InKey, double InCurrentTime) override;
 
 	virtual void GetKeyPositions(TArrayView<const FKeyHandle> InKeys, TArrayView<FKeyPosition> OutKeyPositions) const override;
 	virtual void SetKeyPositions(TArrayView<const FKeyHandle> InKeys, TArrayView<const FKeyPosition> InKeyPositions, EPropertyChangeType::Type ChangeType) override;
diff --git a/Engine/Source/Editor/MovieSceneTools/Private/Channels/ChannelCurveModel.cpp b/Engine/Source/Editor/MovieSceneTools/Private/Channels/ChannelCurveModel.cpp
--- a/Engine/Source/Editor/MovieSceneTools/Private/Channels/ChannelCurveModel.cpp
+++ b/Engine/Source/Editor/MovieSceneTools/Private/Channels/ChannelCurveModel.cpp
@@ -523,7 +523,7 @@ void FChannelCurveModel<ChannelType, ChannelValue, KeyType>::AddKeys(TArrayView<
 }
 
 template <class ChannelType, class ChannelValue, class KeyType>
-void FChannelCurveModel<ChannelType, ChannelValue, KeyType>::RemoveKeys(TArrayView<const FKeyHandle> InKeys)
+void FChannelCurveModel<ChannelType, ChannelValue, KeyType>::RemoveKeys(TArrayView<const FKeyHandle> InKeys, double InCurrentTime)
 {
 	ChannelType* Channel = ChannelHandle.Get();
 	UMovieSceneSection* Section = WeakSection.Get();
@@ -533,11 +533,14 @@ void FChannelCurveModel<ChannelType, ChannelValue, KeyType>::RemoveKeys(TArrayVi
 		SignedOwner = WeakSection.Get();
 	}
 
+
 	if (Channel && Section && !IsReadOnly())
 	{
 		SignedOwner->Modify();
 
-		TMovieSceneChannelData<ChannelValue> ChannelData = Channel->GetData();
+		TMovieSceneChannelData<ChannelValue> ChannelData = Channel->GetData();		
+		double CurrentValue = 0.0;
+		Evaluate(InCurrentTime, CurrentValue);
 
 		for (FKeyHandle Handle : InKeys)
 		{
@@ -547,6 +550,13 @@ void FChannelCurveModel<ChannelType, ChannelValue, KeyType>::RemoveKeys(TArrayVi
 				ChannelData.RemoveKey(KeyIndex);
 			}
 		}
+		if (Channel->GetNumKeys() == 0)
+		{
+			using CurveValueType = typename ChannelType::CurveValueType;
+			CurveValueType ValueAsCurveValue = (CurveValueType)(CurrentValue); //needed for double to other type conversions
+			Channel->SetDefault(ValueAsCurveValue);
+		}
+
 		Channel->PostEditChange();
 		if (WeakSequencer.IsValid())
 		{
diff --git a/Engine/Source/Editor/MovieSceneTools/Private/Channels/ConstraintChannelCurveModel.cpp b/Engine/Source/Editor/MovieSceneTools/Private/Channels/ConstraintChannelCurveModel.cpp
--- a/Engine/Source/Editor/MovieSceneTools/Private/Channels/ConstraintChannelCurveModel.cpp
+++ b/Engine/Source/Editor/MovieSceneTools/Private/Channels/ConstraintChannelCurveModel.cpp
@@ -160,7 +160,7 @@ bool FConstraintChannelCurveModel::Evaluate(double Time, double& OutValue) const
 	return false;
 }
 
-void FConstraintChannelCurveModel::RemoveKeys(TArrayView<const FKeyHandle> InKeys)
+void FConstraintChannelCurveModel::RemoveKeys(TArrayView<const FKeyHandle> InKeys, double InCurrentTime)
 {
 	FMovieSceneConstraintChannel* Channel = GetChannel();
 	UMovieSceneSection* Section = WeakSection.Get();
diff --git a/Engine/Source/Editor/MovieSceneTools/Private/Channels/EventChannelCurveModel.cpp b/Engine/Source/Editor/MovieSceneTools/Private/Channels/EventChannelCurveModel.cpp
--- a/Engine/Source/Editor/MovieSceneTools/Private/Channels/EventChannelCurveModel.cpp
+++ b/Engine/Source/Editor/MovieSceneTools/Private/Channels/EventChannelCurveModel.cpp
@@ -95,7 +95,7 @@ bool FEventChannelCurveModel::Evaluate(double Time, double& OutValue) const
 	return false;
 }
 
-void FEventChannelCurveModel::RemoveKeys(TArrayView<const FKeyHandle> InKeys)
+void FEventChannelCurveModel::RemoveKeys(TArrayView<const FKeyHandle> InKeys, double InCurrentTime)
 {
 	FMovieSceneEventChannel* Channel = ChannelHandle.Get();
 	UMovieSceneSection*      Section = WeakSection.Get();
diff --git a/Engine/Source/Editor/MovieSceneTools/Private/Channels/EventChannelCurveModel.h b/Engine/Source/Editor/MovieSceneTools/Private/Channels/EventChannelCurveModel.h
--- a/Engine/Source/Editor/MovieSceneTools/Private/Channels/EventChannelCurveModel.h
+++ b/Engine/Source/Editor/MovieSceneTools/Private/Channels/EventChannelCurveModel.h
@@ -56,7 +56,7 @@ public:
 	virtual void GetNeighboringKeys(const FKeyHandle InKeyHandle, TOptional<FKeyHandle>& OutPreviousKeyHandle, TOptional<FKeyHandle>& OutNextKeyHandle) const override {}
 	virtual bool Evaluate(double ProspectiveTime, double& OutValue) const override;
 	virtual void AddKeys(TArrayView<const FKeyPosition> InKeyPositions, TArrayView<const FKeyAttributes> InAttributes, TArrayView<TOptional<FKeyHandle>>* OutKeyHandles) override;
-	virtual void RemoveKeys(TArrayView<const FKeyHandle> InKeys) override;
+	virtual void RemoveKeys(TArrayView<const FKeyHandle> InKeys, double InCurrentTime) override;
 
 	virtual void CreateKeyProxies(TArrayView<const FKeyHandle> InKeyHandles, TArrayView<UObject*> OutObjects) override;
 
diff --git a/Engine/Source/Editor/MovieSceneTools/Public/Channels/ChannelCurveModel.h b/Engine/Source/Editor/MovieSceneTools/Public/Channels/ChannelCurveModel.h
--- a/Engine/Source/Editor/MovieSceneTools/Public/Channels/ChannelCurveModel.h
+++ b/Engine/Source/Editor/MovieSceneTools/Public/Channels/ChannelCurveModel.h
@@ -62,7 +62,7 @@ public:
 	virtual void GetNeighboringKeys(const FKeyHandle InKeyHandle, TOptional<FKeyHandle>& OutPreviousKeyHandle, TOptional<FKeyHandle>& OutNextKeyHandle) const override;
 	virtual bool Evaluate(double ProspectiveTime, double& OutValue) const override;
 	virtual void AddKeys(TArrayView<const FKeyPosition> InKeyPositions, TArrayView<const FKeyAttributes> InAttributes, TArrayView<TOptional<FKeyHandle>>* OutKeyHandles) override;
-	virtual void RemoveKeys(TArrayView<const FKeyHandle> InKeys) override;
+	virtual void RemoveKeys(TArrayView<const FKeyHandle> InKeys, double InCurrentTime) override;
 
 	virtual FTransform2d GetCurveTransform() const override;
 	virtual bool IsReadOnly() const override;
diff --git a/Engine/Source/Editor/MovieSceneTools/Public/Channels/ConstraintChannelCurveModel.h b/Engine/Source/Editor/MovieSceneTools/Public/Channels/ConstraintChannelCurveModel.h
--- a/Engine/Source/Editor/MovieSceneTools/Public/Channels/ConstraintChannelCurveModel.h
+++ b/Engine/Source/Editor/MovieSceneTools/Public/Channels/ConstraintChannelCurveModel.h
@@ -62,7 +62,7 @@ public:
 	virtual void GetNeighboringKeys(const FKeyHandle InKeyHandle, TOptional<FKeyHandle>& OutPreviousKeyHandle, TOptional<FKeyHandle>& OutNextKeyHandle) const override {}
 	virtual bool Evaluate(double ProspectiveTime, double& OutValue) const override;
 	virtual void AddKeys(TArrayView<const FKeyPosition> InKeyPositions, TArrayView<const FKeyAttributes> InAttributes, TArrayView<TOptional<FKeyHandle>>* OutKeyHandles) override;
-	virtual void RemoveKeys(TArrayView<const FKeyHandle> InKeys) override;
+	virtual void RemoveKeys(TArrayView<const FKeyHandle> InKeys, double InCurrentTime) override;
 
 	virtual void CreateKeyProxies(TArrayView<const FKeyHandle> InKeyHandles, TArrayView<UObject*> OutObjects) override;
 	
diff --git a/Engine/Source/Editor/MovieSceneTools/Public/Channels/PiecewiseCurveModel.h b/Engine/Source/Editor/MovieSceneTools/Public/Channels/PiecewiseCurveModel.h
--- a/Engine/Source/Editor/MovieSceneTools/Public/Channels/PiecewiseCurveModel.h
+++ b/Engine/Source/Editor/MovieSceneTools/Public/Channels/PiecewiseCurveModel.h
@@ -37,7 +37,7 @@ public:
 	void AddKeys(TArrayView<const FKeyPosition> InPositions, TArrayView<const FKeyAttributes> InAttributes, TArrayView<TOptional<FKeyHandle>>* OutKeyHandles) override
 	{
 	}
-	void RemoveKeys(TArrayView<const FKeyHandle> InKeys) override
+	void RemoveKeys(TArrayView<const FKeyHandle> InKeys, double InCurrentTime) override
 	{
 	}
 	void GetKeyPositions(TArrayView<const FKeyHandle> InKeys, TArrayView<FKeyPosition> OutKeyPositions) const override
