Curve Editor: MMB clicking an existing key no longer resets the key

diff --git a/Engine/Source/Editor/CurveEditor/Private/Views/SInteractiveCurveEditorView.cpp b/Engine/Source/Editor/CurveEditor/Private/Views/SInteractiveCurveEditorView.cpp
--- a/Engine/Source/Editor/CurveEditor/Private/Views/SInteractiveCurveEditorView.cpp
+++ b/Engine/Source/Editor/CurveEditor/Private/Views/SInteractiveCurveEditorView.cpp
@@ -1609,8 +1609,11 @@ FReply SInteractiveCurveEditorView::OnMouseButtonDown(const FGeometry& MyGeometr
 				// Don't allow adding keys when shift is held down with selected keys since that is for dragging keys in a constrained axis
 				const bool bDraggingKeys = MouseEvent.IsShiftDown() && !CurveEditor->GetSelection().GetAll().IsEmpty();
 
+				// Don't create a new key if there is already a key or transform handle in place
+				const bool bKeyAlreadyExists = HitPoint(MousePixel).IsSet();
+
 				FCurveModel* CurveToAddTo = CurveEditor->FindCurve(HoveredCurve.GetValue());
-				if (CurveToAddTo && !CurveToAddTo->IsReadOnly() && !bDraggingKeys)
+				if (CurveToAddTo && !CurveToAddTo->IsReadOnly() && !bDraggingKeys && !bKeyAlreadyExists)
 				{
 					FScopedTransaction Transaction(LOCTEXT("InsertKey", "Insert Key"));
 
