Curve Editor: Toolbar UX changes

Tools are now in a combo button.

View options:
- View modes now in view combo button
- Framing tools now in view combo button
- Add "View" label to view options combo button
- Move tools to be after view options

Curves section:
- Add new combo button for curves
- Move pre-post infinity to combo button
- Filters are in this menu
- Euler is surfaced to this section

Tangents section:
- Add new combo button for the tangent modes

diff --git a/Engine/Plugins/Editor/CurveEditorTools/Source/CurveEditorTools/Private/CurveEditorToolCommands.cpp b/Engine/Plugins/Editor/CurveEditorTools/Source/CurveEditorTools/Private/CurveEditorToolCommands.cpp
--- a/Engine/Plugins/Editor/CurveEditorTools/Source/CurveEditorTools/Private/CurveEditorToolCommands.cpp
+++ b/Engine/Plugins/Editor/CurveEditorTools/Source/CurveEditorTools/Private/CurveEditorToolCommands.cpp
@@ -11,10 +11,10 @@ void FCurveEditorToolCommands::RegisterCommands()
 	UI_COMMAND(SetFocusPlaybackRange, "Focus Playback Range", "Focuses the Curve Editor on the current Playback Range with zoom based on visible curves.", EUserInterfaceActionType::Button, FInputChord(EKeys::A));
 
 	// Tool Modes
-	UI_COMMAND(ActivateTransformTool, "Transform Tool", "Activates the Transform tool which allows translation, scale and rotation of selected keys.", EUserInterfaceActionType::ToggleButton, FInputChord(EModifierKey::Control, EKeys::T));
-	UI_COMMAND(ActivateRetimeTool, "Retime Tool", "Activates the Retime tool which allows you to define a one dimensional lattice to non-uniformly rescale key times.", EUserInterfaceActionType::ToggleButton, FInputChord(EModifierKey::Control, EKeys::E));
-	UI_COMMAND(ActivateMultiScaleTool, "Multi Scale Tool", "Activates the Multi Select tool which allows fast scaling of multiple curves from unique pivots.", EUserInterfaceActionType::ToggleButton, FInputChord(EModifierKey::Control, EKeys::M));
-	UI_COMMAND(ActivateLatticeTool, "Lattice Tool", "Activates the Lattice tool which allows fast editing of keys using a grid.", EUserInterfaceActionType::ToggleButton, FInputChord(EModifierKey::Control, EKeys::L));
+	UI_COMMAND(ActivateTransformTool, "Transform", "The Transform tool allows translation, scale and rotation of selected keys.", EUserInterfaceActionType::RadioButton, FInputChord(EModifierKey::Control, EKeys::T));
+	UI_COMMAND(ActivateRetimeTool, "Retime", "The Retime tool allows you to define a one dimensional lattice to non-uniformly rescale key times.", EUserInterfaceActionType::RadioButton, FInputChord(EModifierKey::Control, EKeys::E));
+	UI_COMMAND(ActivateMultiScaleTool, "Multi Scale", "The Multi Select tool allows fast scaling of multiple curves from unique pivots.", EUserInterfaceActionType::RadioButton, FInputChord(EModifierKey::Control, EKeys::M));
+	UI_COMMAND(ActivateLatticeTool, "Lattice", "The Lattice tool allows fast editing of keys using a grid.", EUserInterfaceActionType::RadioButton, FInputChord(EModifierKey::Control, EKeys::L));
 }
 
 #undef LOCTEXT_NAMESPACE
diff --git a/Engine/Plugins/Editor/CurveEditorTools/Source/CurveEditorTools/Private/CurveEditorToolsModule.cpp b/Engine/Plugins/Editor/CurveEditorTools/Source/CurveEditorTools/Private/CurveEditorToolsModule.cpp
--- a/Engine/Plugins/Editor/CurveEditorTools/Source/CurveEditorTools/Private/CurveEditorToolsModule.cpp
+++ b/Engine/Plugins/Editor/CurveEditorTools/Source/CurveEditorTools/Private/CurveEditorToolsModule.cpp
@@ -1,5 +1,6 @@
 // Copyright Epic Games, Inc. All Rights Reserved.
 
+#include "CurveEditorHelpers.h"
 #include "Modules/ModuleManager.h"
 #include "Framework/Commands/UICommandList.h"
 #include "ICurveEditorModule.h"
@@ -70,31 +71,24 @@ TSharedRef<FExtender> FCurveEditorToolsModule::ExtendCurveEditorToolbarMenu(cons
 {
 	struct Local
 	{
-		static void FillToolbarTools(FToolBarBuilder& ToolbarBuilder)
+		static void FillToolbarTools(FMenuBuilder& ToolbarBuilder)
 		{
-			ToolbarBuilder.AddToolBarButton(FCurveEditorToolCommands::Get().ActivateTransformTool);
-			ToolbarBuilder.AddToolBarButton(FCurveEditorToolCommands::Get().ActivateRetimeTool);
-			ToolbarBuilder.AddToolBarButton(FCurveEditorToolCommands::Get().ActivateMultiScaleTool);
-			ToolbarBuilder.AddToolBarButton(FCurveEditorToolCommands::Get().ActivateLatticeTool);
+			ToolbarBuilder.AddMenuEntry(FCurveEditorToolCommands::Get().ActivateTransformTool);
+			ToolbarBuilder.AddMenuEntry(FCurveEditorToolCommands::Get().ActivateRetimeTool);
+			ToolbarBuilder.AddMenuEntry(FCurveEditorToolCommands::Get().ActivateMultiScaleTool);
+			ToolbarBuilder.AddMenuEntry(FCurveEditorToolCommands::Get().ActivateLatticeTool);
 		}
 
-		static void FillToolbarFraming(FToolBarBuilder& ToolbarBuilder)
+		static void FillToolbarFraming(FMenuBuilder& ToolbarBuilder)
 		{
-			ToolbarBuilder.AddToolBarButton(FCurveEditorToolCommands::Get().SetFocusPlaybackTime);
-			ToolbarBuilder.AddToolBarButton(FCurveEditorToolCommands::Get().SetFocusPlaybackRange);
+			ToolbarBuilder.AddMenuEntry(FCurveEditorToolCommands::Get().SetFocusPlaybackTime);
+			ToolbarBuilder.AddMenuEntry(FCurveEditorToolCommands::Get().SetFocusPlaybackRange);
 		}
 	};
 
-	TSharedRef<FExtender> Extender = MakeShared<FExtender>();
-	Extender->AddToolBarExtension("Tools",
-	EExtensionHook::After,
-	CommandList,
-	FToolBarExtensionDelegate::CreateStatic(&Local::FillToolbarTools));
-	
-	Extender->AddToolBarExtension("Framing",
-	EExtensionHook::After,
-	CommandList,
-	FToolBarExtensionDelegate::CreateStatic(&Local::FillToolbarFraming));
+	const TSharedRef<FExtender> Extender = MakeShared<FExtender>();
+	Extender->AddMenuExtension("Tools", EExtensionHook::After, CommandList, FMenuExtensionDelegate::CreateStatic(&Local::FillToolbarTools));
+	Extender->AddMenuExtension("Framing", EExtensionHook::After, CommandList, FMenuExtensionDelegate::CreateStatic(&Local::FillToolbarFraming));
 
 	return Extender;
 }
@@ -111,7 +105,7 @@ void FCurveEditorToolsModule::ShutdownModule()
 	CurveEditorModule.UnregisterToolExtension(RetimeToolHandle);
 	CurveEditorModule.UnregisterToolExtension(MultiScaleToolHandle);
 	CurveEditorModule.UnregisterToolExtension(LatticeToolHandle);
-
+	
 	FCurveEditorToolCommands::Unregister();
 }
 
diff --git a/Engine/Plugins/Editor/CurveEditorTools/Source/CurveEditorTools/Private/Extensions/Tweening/TweenEditorExtension.cpp b/Engine/Plugins/Editor/CurveEditorTools/Source/CurveEditorTools/Private/Extensions/Tweening/TweenEditorExtension.cpp
--- a/Engine/Plugins/Editor/CurveEditorTools/Source/CurveEditorTools/Private/Extensions/Tweening/TweenEditorExtension.cpp
+++ b/Engine/Plugins/Editor/CurveEditorTools/Source/CurveEditorTools/Private/Extensions/Tweening/TweenEditorExtension.cpp
@@ -29,7 +29,7 @@ TSharedPtr<FExtender> FTweenEditorExtension::MakeToolbarExtender(const TSharedRe
 	using namespace TweeningUtilsEditor;
 	TSharedRef<FExtender> Extender = MakeShared<FExtender>();
 	
-	Extender->AddToolBarExtension("Tangents", EExtensionHook::After, InCommandList, FToolBarExtensionDelegate::CreateLambda(
+	Extender->AddToolBarExtension("Adjustment", EExtensionHook::After, InCommandList, FToolBarExtensionDelegate::CreateLambda(
 		[this, WeakCommandList = InCommandList.ToWeakPtr()](FToolBarBuilder& ToolbarBuilder)
 		{
 			const TSharedPtr<FUICommandList> CommandListPin = WeakCommandList.Pin();
diff --git a/Engine/Plugins/Editor/CurveEditorTools/Source/CurveEditorTools/Private/Tools/CurveEditorMultiScaleTool.cpp b/Engine/Plugins/Editor/CurveEditorTools/Source/CurveEditorTools/Private/Tools/CurveEditorMultiScaleTool.cpp
--- a/Engine/Plugins/Editor/CurveEditorTools/Source/CurveEditorTools/Private/Tools/CurveEditorMultiScaleTool.cpp
+++ b/Engine/Plugins/Editor/CurveEditorTools/Source/CurveEditorTools/Private/Tools/CurveEditorMultiScaleTool.cpp
@@ -359,6 +359,21 @@ void FCurveEditorMultiScaleTool::OnToolOptionsUpdated(const FPropertyChangedEven
 	ActiveTransaction.Reset();
 }
 
+FText FCurveEditorMultiScaleTool::GetLabel() const
+{
+	return FCurveEditorToolCommands::Get().ActivateMultiScaleTool->GetLabel();
+}
+
+FText FCurveEditorMultiScaleTool::GetDescription() const
+{
+	return FCurveEditorToolCommands::Get().ActivateMultiScaleTool->GetDescription();
+}
+
+FSlateIcon FCurveEditorMultiScaleTool::GetIcon() const
+{
+	return FCurveEditorToolCommands::Get().ActivateMultiScaleTool->GetIcon();
+}
+
 void FCurveEditorMultiScaleTool::OnDragStart()
 {
 	ActiveTransaction = MakeUnique<FScopedTransaction>(TEXT("CurveEditorMultiScaleTool"), LOCTEXT("CurveEditorMultiScaleToolTransaction", "Multi Scale Key(s)"), nullptr);
diff --git a/Engine/Plugins/Editor/CurveEditorTools/Source/CurveEditorTools/Private/Tools/CurveEditorMultiScaleTool.h b/Engine/Plugins/Editor/CurveEditorTools/Source/CurveEditorTools/Private/Tools/CurveEditorMultiScaleTool.h
--- a/Engine/Plugins/Editor/CurveEditorTools/Source/CurveEditorTools/Private/Tools/CurveEditorMultiScaleTool.h
+++ b/Engine/Plugins/Editor/CurveEditorTools/Source/CurveEditorTools/Private/Tools/CurveEditorMultiScaleTool.h
@@ -113,6 +113,9 @@ public:
 
 	TSharedPtr<FStructOnScope> GetToolOptions() const override { return ToolOptionsOnScope; }
 	virtual void OnToolOptionsUpdated(const FPropertyChangedEvent& InPropertyChangedEvent) override;
+	virtual FText GetLabel() const override;
+	virtual FText GetDescription() const override;
+	virtual FSlateIcon GetIcon() const override;
 	// ~ICurveEditorToolExtension
 
 private:
diff --git a/Engine/Plugins/Editor/CurveEditorTools/Source/CurveEditorTools/Private/Tools/CurveEditorRetimeTool.cpp b/Engine/Plugins/Editor/CurveEditorTools/Source/CurveEditorTools/Private/Tools/CurveEditorRetimeTool.cpp
--- a/Engine/Plugins/Editor/CurveEditorTools/Source/CurveEditorTools/Private/Tools/CurveEditorRetimeTool.cpp
+++ b/Engine/Plugins/Editor/CurveEditorTools/Source/CurveEditorTools/Private/Tools/CurveEditorRetimeTool.cpp
@@ -462,6 +462,21 @@ void FCurveEditorRetimeTool::BindCommands(TSharedRef<FUICommandList> CommandBind
 	}
 }
 
+FText FCurveEditorRetimeTool::GetLabel() const
+{
+	return FCurveEditorToolCommands::Get().ActivateRetimeTool->GetLabel();
+}
+
+FText FCurveEditorRetimeTool::GetDescription() const
+{
+	return FCurveEditorToolCommands::Get().ActivateRetimeTool->GetDescription();
+}
+
+FSlateIcon FCurveEditorRetimeTool::GetIcon() const
+{
+	return FCurveEditorToolCommands::Get().ActivateRetimeTool->GetIcon();
+}
+
 void FCurveEditorRetimeTool::OnDragStart()
 {
 	ActiveTransaction = MakeUnique<FScopedTransaction>(TEXT("CurveEditorRetimeTool"), LOCTEXT("CurveEditorRetimeToolTransaction", "Retime Key(s)"), nullptr);
diff --git a/Engine/Plugins/Editor/CurveEditorTools/Source/CurveEditorTools/Private/Tools/CurveEditorRetimeTool.h b/Engine/Plugins/Editor/CurveEditorTools/Source/CurveEditorTools/Private/Tools/CurveEditorRetimeTool.h
--- a/Engine/Plugins/Editor/CurveEditorTools/Source/CurveEditorTools/Private/Tools/CurveEditorRetimeTool.h
+++ b/Engine/Plugins/Editor/CurveEditorTools/Source/CurveEditorTools/Private/Tools/CurveEditorRetimeTool.h
@@ -80,6 +80,9 @@ public:
 	virtual void OnToolActivated() override;
 	virtual void OnToolDeactivated() override;
 	virtual void BindCommands(TSharedRef<FUICommandList> CommandBindings) override;
+	virtual FText GetLabel() const override;
+	virtual FText GetDescription() const override;
+	virtual FSlateIcon GetIcon() const override;
 	// ~ICurveEditorToolExtension
 
 private:
diff --git a/Engine/Plugins/Editor/CurveEditorTools/Source/CurveEditorTools/Private/Tools/CurveEditorTransformTool.cpp b/Engine/Plugins/Editor/CurveEditorTools/Source/CurveEditorTools/Private/Tools/CurveEditorTransformTool.cpp
--- a/Engine/Plugins/Editor/CurveEditorTools/Source/CurveEditorTools/Private/Tools/CurveEditorTransformTool.cpp
+++ b/Engine/Plugins/Editor/CurveEditorTools/Source/CurveEditorTools/Private/Tools/CurveEditorTransformTool.cpp
@@ -514,6 +514,21 @@ void FCurveEditorTransformTool::OnToolOptionsUpdated(const FPropertyChangedEvent
 	ActiveTransaction.Reset();
 }
 
+FText FCurveEditorTransformTool::GetLabel() const
+{
+	return FCurveEditorToolCommands::Get().ActivateTransformTool->GetLabel();
+}
+
+FText FCurveEditorTransformTool::GetDescription() const
+{
+	return FCurveEditorToolCommands::Get().ActivateTransformTool->GetDescription();
+}
+
+FSlateIcon FCurveEditorTransformTool::GetIcon() const
+{
+	return FCurveEditorToolCommands::Get().ActivateTransformTool->GetIcon();
+}
+
 void FCurveEditorTransformTool::OnPaint(const FPaintArgs& Args, const FGeometry& AllottedGeometry, const FSlateRect& MyCullingRect, FSlateWindowElementList& OutDrawElements, int32 PaintOnLayerId, const FWidgetStyle& InWidgetStyle, bool bParentEnabled) const
 {
 	// This Geometry represents the Marquee size, but also the offset into the window
diff --git a/Engine/Plugins/Editor/CurveEditorTools/Source/CurveEditorTools/Private/Tools/CurveEditorTransformTool.h b/Engine/Plugins/Editor/CurveEditorTools/Source/CurveEditorTools/Private/Tools/CurveEditorTransformTool.h
--- a/Engine/Plugins/Editor/CurveEditorTools/Source/CurveEditorTools/Private/Tools/CurveEditorTransformTool.h
+++ b/Engine/Plugins/Editor/CurveEditorTools/Source/CurveEditorTools/Private/Tools/CurveEditorTransformTool.h
@@ -134,6 +134,9 @@ public:
 	virtual void OnFocusLost(const FFocusEvent& InFocusEvent);
 	TSharedPtr<FStructOnScope> GetToolOptions() const override { return ToolOptionsOnScope; }
 	virtual void OnToolOptionsUpdated(const FPropertyChangedEvent& InPropertyChangedEvent) override;
+	virtual FText GetLabel() const override;
+	virtual FText GetDescription() const override;
+	virtual FSlateIcon GetIcon() const override;
 	// ~ICurveEditorToolExtension
 
 private:
diff --git a/Engine/Plugins/Editor/CurveEditorTools/Source/CurveEditorTools/Private/Tools/Lattice/CurveEditorLatticeTool.cpp b/Engine/Plugins/Editor/CurveEditorTools/Source/CurveEditorTools/Private/Tools/Lattice/CurveEditorLatticeTool.cpp
--- a/Engine/Plugins/Editor/CurveEditorTools/Source/CurveEditorTools/Private/Tools/Lattice/CurveEditorLatticeTool.cpp
+++ b/Engine/Plugins/Editor/CurveEditorTools/Source/CurveEditorTools/Private/Tools/Lattice/CurveEditorLatticeTool.cpp
@@ -221,6 +221,21 @@ void FCurveEditorLatticeTool::BindCommands(TSharedRef<FUICommandList> CommandBin
 	}
 }
 
+FText FCurveEditorLatticeTool::GetLabel() const
+{
+	return FCurveEditorToolCommands::Get().ActivateLatticeTool->GetLabel();
+}
+
+FText FCurveEditorLatticeTool::GetDescription() const
+{
+	return FCurveEditorToolCommands::Get().ActivateLatticeTool->GetDescription();
+}
+
+FSlateIcon FCurveEditorLatticeTool::GetIcon() const
+{
+	return FCurveEditorToolCommands::Get().ActivateLatticeTool->GetIcon();
+}
+
 void FCurveEditorLatticeTool::AddReferencedObjects(FReferenceCollector& Collector)
 {
 	Collector.AddReferencedObject(UndoObject);
diff --git a/Engine/Plugins/Editor/CurveEditorTools/Source/CurveEditorTools/Private/Tools/Lattice/CurveEditorLatticeTool.h b/Engine/Plugins/Editor/CurveEditorTools/Source/CurveEditorTools/Private/Tools/Lattice/CurveEditorLatticeTool.h
--- a/Engine/Plugins/Editor/CurveEditorTools/Source/CurveEditorTools/Private/Tools/Lattice/CurveEditorLatticeTool.h
+++ b/Engine/Plugins/Editor/CurveEditorTools/Source/CurveEditorTools/Private/Tools/Lattice/CurveEditorLatticeTool.h
@@ -40,6 +40,9 @@ public:
 	virtual void OnToolActivated() override;
 	virtual void OnToolDeactivated() override;
 	virtual void BindCommands(TSharedRef<FUICommandList> CommandBindings) override;
+	virtual FText GetLabel() const override;
+	virtual FText GetDescription() const override;
+	virtual FSlateIcon GetIcon() const override;
 	//~ End ICurveEditorToolExtension Interface
 
 	//~ Begin FGCObject Interface
diff --git a/Engine/Source/Editor/CurveEditor/Private/CurveEditorCommands.cpp b/Engine/Source/Editor/CurveEditor/Private/CurveEditorCommands.cpp
--- a/Engine/Source/Editor/CurveEditor/Private/CurveEditorCommands.cpp
+++ b/Engine/Source/Editor/CurveEditor/Private/CurveEditorCommands.cpp
@@ -22,12 +22,12 @@ void FCurveEditorCommands::RegisterCommands()
 	UI_COMMAND(ToggleExpandCollapseNodes, "Expand/Collapse Nodes", "Toggle expand or collapse selected nodes", EUserInterfaceActionType::Button, FInputChord(EKeys::V) );
 	UI_COMMAND(ToggleExpandCollapseNodesAndDescendants, "Expand/Collapse Nodes and Descendants", "Toggle expand or collapse selected nodes and descendants", EUserInterfaceActionType::Button, FInputChord(EModifierKey::Shift, EKeys::V) );
 
-	UI_COMMAND(InterpolationConstant, "Constant", "Constant interpolation", EUserInterfaceActionType::ToggleButton, FInputChord(EKeys::Five));
-	UI_COMMAND(InterpolationLinear, "Linear", "Linear interpolation", EUserInterfaceActionType::ToggleButton, FInputChord(EKeys::Four));
-	UI_COMMAND(InterpolationCubicSmartAuto, "Smart Auto", "Cubic interpolation - Smart Automatic tangents", EUserInterfaceActionType::ToggleButton, FInputChord(EKeys::Zero));
-	UI_COMMAND(InterpolationCubicAuto, "Auto", "Cubic interpolation - Automatic tangents", EUserInterfaceActionType::ToggleButton, FInputChord(EKeys::One));
-	UI_COMMAND(InterpolationCubicUser, "User", "Cubic interpolation - User flat tangents", EUserInterfaceActionType::ToggleButton, FInputChord(EKeys::Two));
-	UI_COMMAND(InterpolationCubicBreak, "Break", "Cubic interpolation - User broken tangents", EUserInterfaceActionType::ToggleButton, FInputChord(EKeys::Three));
+	UI_COMMAND(InterpolationConstant, "Constant", "Constant interpolation", EUserInterfaceActionType::RadioButton, FInputChord(EKeys::Five));
+	UI_COMMAND(InterpolationLinear, "Linear", "Linear interpolation", EUserInterfaceActionType::RadioButton, FInputChord(EKeys::Four));
+	UI_COMMAND(InterpolationCubicSmartAuto, "Smart Auto", "Cubic interpolation - Smart Automatic tangents", EUserInterfaceActionType::RadioButton, FInputChord(EKeys::Zero));
+	UI_COMMAND(InterpolationCubicAuto, "Auto", "Cubic interpolation - Automatic tangents", EUserInterfaceActionType::RadioButton, FInputChord(EKeys::One));
+	UI_COMMAND(InterpolationCubicUser, "User", "Cubic interpolation - User flat tangents", EUserInterfaceActionType::RadioButton, FInputChord(EKeys::Two));
+	UI_COMMAND(InterpolationCubicBreak, "Break", "Cubic interpolation - User broken tangents", EUserInterfaceActionType::RadioButton, FInputChord(EKeys::Three));
 	UI_COMMAND(InterpolationToggleWeighted, "Weighted Tangents", "Toggle weighted tangents for cubic interpolation modes. Only supported on some curve types", EUserInterfaceActionType::ToggleButton, FInputChord(EModifierKey::Control, EKeys::W));
 
 	UI_COMMAND(FlattenTangents, "Flatten", "Flatten Tangents", EUserInterfaceActionType::Button, FInputChord(EKeys::Six));
@@ -76,12 +76,12 @@ void FCurveEditorCommands::RegisterCommands()
 	UI_COMMAND(SetCurveColorsForSelected, "Set Curve Color For Selected", "Set the chosen color on the selected curves. Note they are stored in the Level Sequence Actor Editor Preferences.", EUserInterfaceActionType::Button, FInputChord());
 
 	// Graph Viewing Modes
-	UI_COMMAND(SetViewModeAbsolute, "Absolute View Mode", "Absolute view displays all curves overlapping with the Y axis proportionally scaled.", EUserInterfaceActionType::ToggleButton, FInputChord());
-	UI_COMMAND(SetViewModeStacked, "Stacked View Mode", "Stacked view displays each curve in its own graph with the Y axis normalized [-1, 1].", EUserInterfaceActionType::ToggleButton, FInputChord());
-	UI_COMMAND(SetViewModeNormalized, "Normalized View Mode", "Normalized view displays all curves overlapping with the Y axis normalized [-1, 1].", EUserInterfaceActionType::ToggleButton, FInputChord());
+	UI_COMMAND(SetViewModeAbsolute, "Absolute View Mode", "Absolute view displays all curves overlapping with the Y axis proportionally scaled.", EUserInterfaceActionType::RadioButton, FInputChord());
+	UI_COMMAND(SetViewModeStacked, "Stacked View Mode", "Stacked view displays each curve in its own graph with the Y axis normalized [-1, 1].", EUserInterfaceActionType::RadioButton, FInputChord());
+	UI_COMMAND(SetViewModeNormalized, "Normalized View Mode", "Normalized view displays all curves overlapping with the Y axis normalized [-1, 1].", EUserInterfaceActionType::RadioButton, FInputChord());
 
 	// Deactivate the currently active tool
-	UI_COMMAND(DeactivateCurrentTool, "Deactivate Tool", "Deactivates the current tool and returns to just supporting selection.", EUserInterfaceActionType::ToggleButton, FInputChord(EKeys::Q));
+	UI_COMMAND(DeactivateCurrentTool, "Selection", "Deactivates the current tool and returns to just supporting selection.", EUserInterfaceActionType::RadioButton, FInputChord(EKeys::Q));
 
 	// User Implementable Filter window
 	UI_COMMAND(OpenUserImplementableFilterWindow, "Filter...", "Opens a window which lets you choose from user implementable filter classes with advanced settings.", EUserInterfaceActionType::Button, FInputChord());
diff --git a/Engine/Source/Editor/CurveEditor/Private/SCurveEditorPanel.cpp b/Engine/Source/Editor/CurveEditor/Private/SCurveEditorPanel.cpp
--- a/Engine/Source/Editor/CurveEditor/Private/SCurveEditorPanel.cpp
+++ b/Engine/Source/Editor/CurveEditor/Private/SCurveEditorPanel.cpp
@@ -395,6 +395,7 @@ void SCurveEditorPanel::BindCommands()
 		}
 	}
 
+	const FCanExecuteAction CanExtrapolate = FCanExecuteAction::CreateLambda([this](){ return CanSetKeyInterpolation(); });
 	// Pre Extrapolation Modes
 	{
 		FExecuteAction SetCycle           = FExecuteAction::CreateSP(this, &SCurveEditorPanel::SetCurveAttributes, FCurveAttributes().SetPreExtrapolation(RCCE_Cycle),           LOCTEXT("SetPreExtrapCycle",           "Set Pre Extrapolation (Cycle)"));
@@ -409,11 +410,11 @@ void SCurveEditorPanel::BindCommands()
 		FIsActionChecked IsLinearCommon          = FIsActionChecked::CreateSP(this, &SCurveEditorPanel::CompareCommonPreExtrapolationMode, RCCE_Linear);
 		FIsActionChecked IsConstantCommon        = FIsActionChecked::CreateSP(this, &SCurveEditorPanel::CompareCommonPreExtrapolationMode, RCCE_Constant);
 
-		CommandList->MapAction(FCurveEditorCommands::Get().SetPreInfinityExtrapCycle, SetCycle, FCanExecuteAction(), IsCycleCommon);
-		CommandList->MapAction(FCurveEditorCommands::Get().SetPreInfinityExtrapCycleWithOffset, SetCycleWithOffset, FCanExecuteAction(), IsCycleWithOffsetCommon);
-		CommandList->MapAction(FCurveEditorCommands::Get().SetPreInfinityExtrapOscillate, SetOscillate, FCanExecuteAction(), IsOscillateCommon);
-		CommandList->MapAction(FCurveEditorCommands::Get().SetPreInfinityExtrapLinear, SetLinear, FCanExecuteAction(), IsLinearCommon);
-		CommandList->MapAction(FCurveEditorCommands::Get().SetPreInfinityExtrapConstant, SetConstant, FCanExecuteAction(), IsConstantCommon);
+		CommandList->MapAction(FCurveEditorCommands::Get().SetPreInfinityExtrapCycle, SetCycle, CanExtrapolate, IsCycleCommon);
+		CommandList->MapAction(FCurveEditorCommands::Get().SetPreInfinityExtrapCycleWithOffset, SetCycleWithOffset, CanExtrapolate, IsCycleWithOffsetCommon);
+		CommandList->MapAction(FCurveEditorCommands::Get().SetPreInfinityExtrapOscillate, SetOscillate, CanExtrapolate, IsOscillateCommon);
+		CommandList->MapAction(FCurveEditorCommands::Get().SetPreInfinityExtrapLinear, SetLinear, CanExtrapolate, IsLinearCommon);
+		CommandList->MapAction(FCurveEditorCommands::Get().SetPreInfinityExtrapConstant, SetConstant, CanExtrapolate, IsConstantCommon);
 	}
 
 	// Post Extrapolation Modes
@@ -430,11 +431,11 @@ void SCurveEditorPanel::BindCommands()
 		FIsActionChecked IsLinearCommon          = FIsActionChecked::CreateSP(this, &SCurveEditorPanel::CompareCommonPostExtrapolationMode, RCCE_Linear);
 		FIsActionChecked IsConstantCommon        = FIsActionChecked::CreateSP(this, &SCurveEditorPanel::CompareCommonPostExtrapolationMode, RCCE_Constant);
 
-		CommandList->MapAction(FCurveEditorCommands::Get().SetPostInfinityExtrapCycle, SetCycle, FCanExecuteAction(), IsCycleCommon);
-		CommandList->MapAction(FCurveEditorCommands::Get().SetPostInfinityExtrapCycleWithOffset, SetCycleWithOffset, FCanExecuteAction(), IsCycleWithOffsetCommon);
-		CommandList->MapAction(FCurveEditorCommands::Get().SetPostInfinityExtrapOscillate, SetOscillate, FCanExecuteAction(), IsOscillateCommon);
-		CommandList->MapAction(FCurveEditorCommands::Get().SetPostInfinityExtrapLinear, SetLinear, FCanExecuteAction(), IsLinearCommon);
-		CommandList->MapAction(FCurveEditorCommands::Get().SetPostInfinityExtrapConstant, SetConstant, FCanExecuteAction(), IsConstantCommon);
+		CommandList->MapAction(FCurveEditorCommands::Get().SetPostInfinityExtrapCycle, SetCycle, CanExtrapolate, IsCycleCommon);
+		CommandList->MapAction(FCurveEditorCommands::Get().SetPostInfinityExtrapCycleWithOffset, SetCycleWithOffset, CanExtrapolate, IsCycleWithOffsetCommon);
+		CommandList->MapAction(FCurveEditorCommands::Get().SetPostInfinityExtrapOscillate, SetOscillate, CanExtrapolate, IsOscillateCommon);
+		CommandList->MapAction(FCurveEditorCommands::Get().SetPostInfinityExtrapLinear, SetLinear, CanExtrapolate, IsLinearCommon);
+		CommandList->MapAction(FCurveEditorCommands::Get().SetPostInfinityExtrapConstant, SetConstant, CanExtrapolate, IsConstantCommon);
 	}
 
 	// Absolute, Stacked and Normalized views.
@@ -905,6 +906,135 @@ void SCurveEditorPanel::UpdateEditBox()
 	KeyDetailsView->GetPropertyRowGenerator()->SetObjects(AllEditObjects);
 }
 
+TSharedRef<SWidget> SCurveEditorPanel::MakeTangentModeMenu()
+{
+	FMenuBuilder MenuBuilder(true, GetCommands());
+	
+	const int32 SupportedTangentTypes = CurveEditor->GetSupportedTangentTypes();
+	if (SupportedTangentTypes & (int32)ECurveEditorTangentTypes::InterpolationCubicSmartAuto)
+	{
+		MenuBuilder.AddMenuEntry(FCurveEditorCommands::Get().InterpolationCubicSmartAuto);
+	};
+	if (SupportedTangentTypes & (int32)ECurveEditorTangentTypes::InterpolationCubicAuto)
+	{
+		MenuBuilder.AddMenuEntry(FCurveEditorCommands::Get().InterpolationCubicAuto);
+	};
+	if (SupportedTangentTypes & (int32)ECurveEditorTangentTypes::InterpolationCubicUser)
+	{
+		MenuBuilder.AddMenuEntry(FCurveEditorCommands::Get().InterpolationCubicUser);
+	}
+	if (SupportedTangentTypes & (int32)ECurveEditorTangentTypes::InterpolationCubicBreak)
+	{
+		MenuBuilder.AddMenuEntry(FCurveEditorCommands::Get().InterpolationCubicBreak);
+	}
+	if (SupportedTangentTypes & (int32)ECurveEditorTangentTypes::InterpolationLinear)
+	{
+		MenuBuilder.AddMenuEntry(FCurveEditorCommands::Get().InterpolationLinear);
+	}
+	if (SupportedTangentTypes & (int32)ECurveEditorTangentTypes::InterpolationConstant)
+	{
+		MenuBuilder.AddMenuEntry(FCurveEditorCommands::Get().InterpolationConstant);
+	}
+
+	return MenuBuilder.MakeWidget();
+}
+
+FText SCurveEditorPanel::GetTangentModeLabel() const
+{
+	switch (DetermineTangentMode())
+	{
+	case ETangentModeComboState::NoSelection: return LOCTEXT("TangentMode.NoSelection.Label", "Mode");
+	case ETangentModeComboState::Constant: return LOCTEXT("TangentMode.Constant.Label", "Constant");
+	case ETangentModeComboState::Linear: return LOCTEXT("TangentMode.Linear.Label", "Linear");
+	case ETangentModeComboState::CubicAuto: return LOCTEXT("TangentMode.CubicAuto.Label", "Cubic (auto)");
+	case ETangentModeComboState::CubicSmartAuto: return LOCTEXT("TangentMode.CubicSmartAuto.Label", "Cubic (smart)");
+	case ETangentModeComboState::CubicUser: return LOCTEXT("TangentMode.CubicUser.Label", "Cubic (flat)");
+	case ETangentModeComboState::CubicBreak: return LOCTEXT("TangentMode.CubicBreak.Label", "Cubic (broken)");
+	case ETangentModeComboState::Mixed: return LOCTEXT("TangentMode.Mixed.Label", "Mixed");
+	default: checkNoEntry(); return FText::GetEmpty();
+	}
+}
+
+FText SCurveEditorPanel::GetTangentModeTooltip() const
+{
+	switch (DetermineTangentMode())
+	{
+	case ETangentModeComboState::NoSelection: return LOCTEXT("TangentMode.NoSelection.Tooltip", "Select keys to set tangent mode");
+	case ETangentModeComboState::Constant: return FCurveEditorCommands::Get().InterpolationConstant->GetDescription();
+	case ETangentModeComboState::Linear: return FCurveEditorCommands::Get().InterpolationLinear->GetDescription();
+	case ETangentModeComboState::CubicAuto: return FCurveEditorCommands::Get().InterpolationCubicAuto->GetDescription();
+	case ETangentModeComboState::CubicSmartAuto: return FCurveEditorCommands::Get().InterpolationCubicSmartAuto->GetDescription();
+	case ETangentModeComboState::CubicUser: return FCurveEditorCommands::Get().InterpolationCubicUser->GetDescription();
+	case ETangentModeComboState::CubicBreak: return FCurveEditorCommands::Get().InterpolationCubicBreak->GetDescription();
+	case ETangentModeComboState::Mixed: return LOCTEXT("TangentMode.Mixed.Tooltip", "The selected keys have different tangent modes.");
+	default: checkNoEntry(); return FText::GetEmpty();
+	}
+}
+
+FSlateIcon SCurveEditorPanel::GetTangentModeIcon() const
+{
+	switch (DetermineTangentMode())
+	{
+	case ETangentModeComboState::NoSelection: return FSlateIcon(FAppStyle::Get().GetStyleSetName(), "GenericCurveEditor.InterpolationNoSelection");
+	case ETangentModeComboState::Constant: return FSlateIcon(FAppStyle::Get().GetStyleSetName(), "GenericCurveEditor.InterpolationConstant");
+	case ETangentModeComboState::Linear: return FSlateIcon(FAppStyle::Get().GetStyleSetName(), "GenericCurveEditor.InterpolationLinear");
+	case ETangentModeComboState::CubicAuto: return FSlateIcon(FAppStyle::Get().GetStyleSetName(), "GenericCurveEditor.InterpolationCubicAuto");
+	case ETangentModeComboState::CubicSmartAuto: return FSlateIcon(FAppStyle::Get().GetStyleSetName(), "GenericCurveEditor.InterpolationCubicSmartAuto");
+	case ETangentModeComboState::CubicUser: return FSlateIcon(FAppStyle::Get().GetStyleSetName(), "GenericCurveEditor.InterpolationCubicUser");
+	case ETangentModeComboState::CubicBreak: return FSlateIcon(FAppStyle::Get().GetStyleSetName(), "GenericCurveEditor.InterpolationCubicBreak");
+	case ETangentModeComboState::Mixed: return FSlateIcon(FAppStyle::Get().GetStyleSetName(), "GenericCurveEditor.InterpolationMixed");
+	default: checkNoEntry(); return FSlateIcon();
+	}
+}
+
+bool SCurveEditorPanel::IsTangentModeComboEnabled() const
+{
+	const ETangentModeComboState ComboState = DetermineTangentMode();
+	return ComboState != ETangentModeComboState::NoSelection;
+}
+
+SCurveEditorPanel::ETangentModeComboState SCurveEditorPanel::DetermineTangentMode() const
+{
+	if (CurveEditor->Selection.IsEmpty())
+	{
+		return ETangentModeComboState::NoSelection;
+	}
+	
+	const bool bIsConstantCommon		= CompareCommonInterpolationMode(RCIM_Constant);
+	const bool bIsLinearCommon			= CompareCommonInterpolationMode(RCIM_Linear);
+	const bool bIsCubicAutoCommon		= CompareCommonTangentMode(RCIM_Cubic, RCTM_Auto);
+	const bool bIsCubicSmartAutoCommon	= CompareCommonTangentMode(RCIM_Cubic, RCTM_SmartAuto);
+	const bool bIsCubicUserCommon		= CompareCommonTangentMode(RCIM_Cubic, RCTM_User);
+	const bool bIsCubicBreakCommon		= CompareCommonTangentMode(RCIM_Cubic, RCTM_Break);
+
+	if (bIsConstantCommon)
+	{
+		return ETangentModeComboState::Constant;
+	}
+	if (bIsLinearCommon)
+	{
+		return ETangentModeComboState::Linear;
+	}
+	if (bIsCubicAutoCommon)
+	{
+		return ETangentModeComboState::CubicAuto;
+	}
+	if (bIsCubicSmartAutoCommon)
+	{
+		return ETangentModeComboState::CubicSmartAuto;
+	}
+	if (bIsCubicUserCommon)
+	{
+		return ETangentModeComboState::CubicUser;
+	}
+	if (bIsCubicBreakCommon)
+	{
+		return ETangentModeComboState::CubicBreak;
+	}
+	
+	return ETangentModeComboState::Mixed;
+}
+
 EVisibility SCurveEditorPanel::GetSplitterVisibility() const
 {
 	return EVisibility::Visible;
@@ -1033,10 +1163,25 @@ void SCurveEditorPanel::EnablePendingFocusOnHovering(const bool InEnabled)
 	PendingFocus.Enable(InEnabled);
 }
 
-TSharedRef<SWidget> SCurveEditorPanel::MakeCurveEditorCurveViewOptionsMenu()
+TSharedRef<SWidget> SCurveEditorPanel::MakeCurveEditorCurveViewOptionsMenu(TSharedPtr<FExtender> InExtender)
 {
-	// This builds the dropdown menu when looking at the Curve View Options combobox.
-	FMenuBuilder MenuBuilder(true, CurveEditor->GetCommands());
+	FMenuBuilder MenuBuilder(true, CurveEditor->GetCommands(), InExtender);
+
+	// Framing
+	MenuBuilder.BeginSection("Framing", LOCTEXT("FrameHeader", "Frame"));
+	{
+		MenuBuilder.AddMenuEntry(FCurveEditorCommands::Get().ZoomToFit);
+	}
+	MenuBuilder.EndSection();
+	
+	// Viewing mode
+	MenuBuilder.BeginSection("CurveViewMode", LOCTEXT("CurveViewModeHeader", "Curve View Mode"));
+	{
+		MenuBuilder.AddMenuEntry(FCurveEditorCommands::Get().SetViewModeAbsolute);
+		MenuBuilder.AddMenuEntry(FCurveEditorCommands::Get().SetViewModeStacked);
+		MenuBuilder.AddMenuEntry(FCurveEditorCommands::Get().SetViewModeNormalized);
+	}
+	MenuBuilder.EndSection();
 
 	MenuBuilder.BeginSection("TangentVisibility", LOCTEXT("CurveEditorMenuTangentVisibilityHeader", "Tangent Visibility"));
 	{
@@ -1071,38 +1216,38 @@ TSharedRef<SWidget> SCurveEditorPanel::MakeCurveEditorCurveViewOptionsMenu()
 	return MenuBuilder.MakeWidget();
 }
 
-TSharedRef<SWidget> SCurveEditorPanel::MakeCurveExtrapolationMenu(const bool bInPostExtrapolation)
+TSharedRef<SWidget> SCurveEditorPanel::MakeCurvesMenu(TSharedPtr<FExtender> InExtender)
 {
-	// This builds the dropdown menu when looking at the Curve View Options combobox.
-	FMenuBuilder MenuBuilder(true, CurveEditor->GetCommands());
+	FMenuBuilder MenuBuilder(true, CurveEditor->GetCommands(), InExtender);
+	AddPreInfinityToMenu(MenuBuilder);
+	AddPostInfinityToMenu(MenuBuilder);
+	return MenuBuilder.MakeWidget();
+}
 
-	if (!bInPostExtrapolation)
+void SCurveEditorPanel::AddPreInfinityToMenu(FMenuBuilder& InMenuBuilder)
+{
+	InMenuBuilder.BeginSection("PreInfinity", LOCTEXT("CurveEditorMenuPreInfinityHeader", "Pre-Infinity"));
 	{
-		MenuBuilder.BeginSection("PreInfinity", LOCTEXT("CurveEditorMenuPreInfinityHeader", "Pre-Infinity"));
-		{
-			MenuBuilder.AddMenuEntry(FCurveEditorCommands::Get().SetPreInfinityExtrapConstant);
-			MenuBuilder.AddMenuEntry(FCurveEditorCommands::Get().SetPreInfinityExtrapCycle);
-			MenuBuilder.AddMenuEntry(FCurveEditorCommands::Get().SetPreInfinityExtrapCycleWithOffset);
-			MenuBuilder.AddMenuEntry(FCurveEditorCommands::Get().SetPreInfinityExtrapLinear);
-			MenuBuilder.AddMenuEntry(FCurveEditorCommands::Get().SetPreInfinityExtrapOscillate);
-		}
-		MenuBuilder.EndSection();
+		InMenuBuilder.AddMenuEntry(FCurveEditorCommands::Get().SetPreInfinityExtrapConstant);
+		InMenuBuilder.AddMenuEntry(FCurveEditorCommands::Get().SetPreInfinityExtrapCycle);
+		InMenuBuilder.AddMenuEntry(FCurveEditorCommands::Get().SetPreInfinityExtrapCycleWithOffset);
+		InMenuBuilder.AddMenuEntry(FCurveEditorCommands::Get().SetPreInfinityExtrapLinear);
+		InMenuBuilder.AddMenuEntry(FCurveEditorCommands::Get().SetPreInfinityExtrapOscillate);
 	}
-	else
+	InMenuBuilder.EndSection();
+}
+
+void SCurveEditorPanel::AddPostInfinityToMenu(FMenuBuilder& InMenuBuilder)
+{
+	InMenuBuilder.BeginSection("PostInfinity", LOCTEXT("CurveEditorMenuPostInfinityHeader", "Post-Infinity"));
 	{
-		MenuBuilder.BeginSection("PostInfinity", LOCTEXT("CurveEditorMenuPostInfinityHeader", "Post-Infinity"));
-		{
-			MenuBuilder.AddMenuEntry(FCurveEditorCommands::Get().SetPostInfinityExtrapConstant);
-			MenuBuilder.AddMenuEntry(FCurveEditorCommands::Get().SetPostInfinityExtrapCycle);
-			MenuBuilder.AddMenuEntry(FCurveEditorCommands::Get().SetPostInfinityExtrapCycleWithOffset);
-			MenuBuilder.AddMenuEntry(FCurveEditorCommands::Get().SetPostInfinityExtrapLinear);
-			MenuBuilder.AddMenuEntry(FCurveEditorCommands::Get().SetPostInfinityExtrapOscillate);
-		}
-		MenuBuilder.EndSection();
+		InMenuBuilder.AddMenuEntry(FCurveEditorCommands::Get().SetPostInfinityExtrapConstant);
+		InMenuBuilder.AddMenuEntry(FCurveEditorCommands::Get().SetPostInfinityExtrapCycle);
+		InMenuBuilder.AddMenuEntry(FCurveEditorCommands::Get().SetPostInfinityExtrapCycleWithOffset);
+		InMenuBuilder.AddMenuEntry(FCurveEditorCommands::Get().SetPostInfinityExtrapLinear);
+		InMenuBuilder.AddMenuEntry(FCurveEditorCommands::Get().SetPostInfinityExtrapOscillate);
 	}
-
-
-	return MenuBuilder.MakeWidget();
+	InMenuBuilder.EndSection();
 }
 
 FSlateIcon SCurveEditorPanel::GetCurveExtrapolationPreIcon() const
@@ -1218,54 +1363,38 @@ TSharedPtr<FExtender> SCurveEditorPanel::GetToolbarExtender()
 			FToolBarBuilder& ToolBarBuilder,
 			TSharedRef<SCurveKeyDetailPanel> InKeyDetailsPanel,
 			TSharedRef<SCurveEditorPanel> InEditorPanel,
-			TSharedRef<UE::CurveEditor::FPromotedFilterContainer> InToolbarPromotedFilters
+			TSharedRef<UE::CurveEditor::FPromotedFilterContainer> InToolbarPromotedFilters,
+			TSharedPtr<FExtender> InExtender
 			)
 		{
 			ToolBarBuilder.BeginSection("View");
 			ToolBarBuilder.BeginStyleOverride("CurveEditorToolbar");
 			{
-				// Dropdown Menu for choosing your viewing mode
-				TAttribute<FSlateIcon> ViewModeIcon;
-				ViewModeIcon.Bind(TAttribute<FSlateIcon>::FGetter::CreateLambda([InEditorPanel] {
-					switch (InEditorPanel->GetViewMode())
-					{
-					case ECurveEditorViewID::Absolute:
-						return FCurveEditorCommands::Get().SetViewModeAbsolute->GetIcon();
-					case ECurveEditorViewID::Stacked:
-						return FCurveEditorCommands::Get().SetViewModeStacked->GetIcon();
-					case ECurveEditorViewID::Normalized:
-						return FCurveEditorCommands::Get().SetViewModeNormalized->GetIcon();
-					default: // EKeyGroupMode::None
-						return FCurveEditorCommands::Get().SetAxisSnappingNone->GetIcon();
-					}
-				}));
-
+				ToolBarBuilder.SetLabelVisibility(EVisibility::Visible); // Show label..
 				ToolBarBuilder.AddComboButton(
 					FUIAction(),
-					FOnGetContent::CreateSP(InEditorPanel, &SCurveEditorPanel::MakeViewModeMenu),
-					LOCTEXT("ViewModeDropdown", "Curve View Modes"),
-					LOCTEXT("ViewModeDropdownToolTip", "Choose the viewing mode for the curves."),
-					ViewModeIcon);
+					FOnGetContent::CreateSP(InEditorPanel, &SCurveEditorPanel::MakeCurveEditorCurveViewOptionsMenu, InExtender),
+					LOCTEXT("CurveEditor.ViewOptions.Label", "View"),
+					LOCTEXT("CurveEditor.ViewOptions.ToolTip", "View Options"),
+					FSlateIcon(FAppStyle::Get().GetStyleSetName(), "Icons.Visibility")
+					);
+				ToolBarBuilder.SetLabelVisibility(EVisibility::Collapsed); // ... but don't affect the other entries
 			}
 			ToolBarBuilder.EndSection();
 
-			ToolBarBuilder.BeginSection("Framing");
-			{
-				// Framing
-				ToolBarBuilder.AddToolBarButton(FCurveEditorCommands::Get().ZoomToFit);
-			}
-			ToolBarBuilder.EndSection();
-
-			ToolBarBuilder.BeginSection("Visibility");
+			ToolBarBuilder.BeginSection("Tools");
 			{
-				// Curve Visibility
+				ToolBarBuilder.BeginStyleOverride("CurveEditorToolBar.ToolsCombo"); // Min size for combo button so size stays consistent independent of label size
+				ToolBarBuilder.SetLabelVisibility(EVisibility::Visible); // Show label..
 				ToolBarBuilder.AddComboButton(
 					FUIAction(),
-					FOnGetContent::CreateSP(InEditorPanel, &SCurveEditorPanel::MakeCurveEditorCurveViewOptionsMenu),
-					LOCTEXT("CurveEditorCurveOptions", "Curves Options"),
-					LOCTEXT("CurveEditorCurveOptionsToolTip", "Curve Options"),
-					FSlateIcon(FAppStyle::Get().GetStyleSetName(), "Icons.Visibility"));
-
+					FOnGetContent::CreateSP(InEditorPanel, &SCurveEditorPanel::MakeToolsComboMenu, InExtender),
+					TAttribute<FText>(InEditorPanel, &SCurveEditorPanel::GetCurrentToolLabel),
+					FText::GetEmpty(),
+					TAttribute<FSlateIcon>(InEditorPanel, &SCurveEditorPanel::GetCurrentToolIcon)
+					);
+				ToolBarBuilder.SetLabelVisibility(EVisibility::Collapsed); // ... but don't affect the other entries
+				ToolBarBuilder.EndStyleOverride();
 			}
 			ToolBarBuilder.EndSection();
 
@@ -1275,14 +1404,29 @@ TSharedPtr<FExtender> SCurveEditorPanel::GetToolbarExtender()
 			}
 			ToolBarBuilder.EndSection();
 
-			ToolBarBuilder.BeginSection("Tools");
-			{
-				ToolBarBuilder.AddToolBarButton(FCurveEditorCommands::Get().DeactivateCurrentTool);
-			}
-			ToolBarBuilder.EndSection();
-
 			ToolBarBuilder.BeginSection("Adjustment");
 			{
+				// Dropdown Menu for choosing your axis snapping for tool movement
+				TAttribute<FSlateIcon> AxisSnappingModeIcon;
+				AxisSnappingModeIcon.Bind(TAttribute<FSlateIcon>::FGetter::CreateLambda([InEditorPanel] {
+					switch (InEditorPanel->GetCurveEditor()->GetAxisSnap().RestrictedAxisList)
+					{
+					case ECurveEditorSnapAxis::CESA_X:
+						return FCurveEditorCommands::Get().SetAxisSnappingHorizontal->GetIcon();
+					case ECurveEditorSnapAxis::CESA_Y:
+						return FCurveEditorCommands::Get().SetAxisSnappingVertical->GetIcon();
+					default: // ECurveEditorSnapAxis::CESA_None
+						return FCurveEditorCommands::Get().SetAxisSnappingNone->GetIcon();
+					}
+				}));
+
+				ToolBarBuilder.AddComboButton(
+					FUIAction(),
+					FOnGetContent::CreateSP(InEditorPanel, &SCurveEditorPanel::MakeAxisSnapMenu),
+					LOCTEXT("AxisSnappingOptions", "Axis Snapping"),
+					LOCTEXT("AxisSnappingOptionsToolTip", "Choose which axes movement tools are locked to."),
+					AxisSnappingModeIcon);
+				
 				// Toggle Button for Time Snapping
 				ToolBarBuilder.AddToolBarButton(FCurveEditorCommands::Get().ToggleInputSnapping);
 
@@ -1307,105 +1451,59 @@ TSharedPtr<FExtender> SCurveEditorPanel::GetToolbarExtender()
 					LOCTEXT("GridSnappingOptionsToolTip", "Choose the spacing between horizontal grid lines."),
 					TAttribute<FSlateIcon>(),
 					true);
-
-
-				// Dropdown Menu for choosing your axis snapping for tool movement
-				TAttribute<FSlateIcon> AxisSnappingModeIcon;
-				AxisSnappingModeIcon.Bind(TAttribute<FSlateIcon>::FGetter::CreateLambda([InEditorPanel] {
-					switch (InEditorPanel->GetCurveEditor()->GetAxisSnap().RestrictedAxisList)
-					{
-					case ECurveEditorSnapAxis::CESA_X:
-						return FCurveEditorCommands::Get().SetAxisSnappingHorizontal->GetIcon();
-					case ECurveEditorSnapAxis::CESA_Y:
-						return FCurveEditorCommands::Get().SetAxisSnappingVertical->GetIcon();
-					default: // ECurveEditorSnapAxis::CESA_None
-						return FCurveEditorCommands::Get().SetAxisSnappingNone->GetIcon();
-					}
-				}));
-
-				ToolBarBuilder.AddComboButton(
-					FUIAction(),
-					FOnGetContent::CreateSP(InEditorPanel, &SCurveEditorPanel::MakeAxisSnapMenu),
-					LOCTEXT("AxisSnappingOptions", "Axis Snapping"),
-					LOCTEXT("AxisSnappingOptionsToolTip", "Choose which axes movement tools are locked to."),
-					AxisSnappingModeIcon);
 			}
 			ToolBarBuilder.EndSection();
 
 			ToolBarBuilder.BeginSection("Tangents");
 			{
-				int32 SupportedTangentTypes = InEditorPanel->CurveEditor->GetSupportedTangentTypes();
-				if (SupportedTangentTypes & (int32)ECurveEditorTangentTypes::InterpolationCubicSmartAuto)
-				{
-					ToolBarBuilder.AddToolBarButton(FCurveEditorCommands::Get().InterpolationCubicSmartAuto);
-				};
-				if (SupportedTangentTypes & (int32)ECurveEditorTangentTypes::InterpolationCubicAuto)
-				{
-					ToolBarBuilder.AddToolBarButton(FCurveEditorCommands::Get().InterpolationCubicAuto);
-				};
-				if (SupportedTangentTypes & (int32)ECurveEditorTangentTypes::InterpolationCubicUser)
-				{
-					ToolBarBuilder.AddToolBarButton(FCurveEditorCommands::Get().InterpolationCubicUser);
-				}
-				if (SupportedTangentTypes & (int32)ECurveEditorTangentTypes::InterpolationCubicBreak)
-				{
-					ToolBarBuilder.AddToolBarButton(FCurveEditorCommands::Get().InterpolationCubicBreak);
-				}
-				if (SupportedTangentTypes & (int32)ECurveEditorTangentTypes::InterpolationLinear)
-				{
-					ToolBarBuilder.AddToolBarButton(FCurveEditorCommands::Get().InterpolationLinear);
-				}
-				if (SupportedTangentTypes & (int32)ECurveEditorTangentTypes::InterpolationConstant)
-				{
-					ToolBarBuilder.AddToolBarButton(FCurveEditorCommands::Get().InterpolationConstant);
-				}
+				ToolBarBuilder.BeginStyleOverride("CurveEditorToolBar.TangentMode"); // Min size for combo button so size stays consistent independent of label size
+				ToolBarBuilder.SetLabelVisibility(EVisibility::Visible); // Show label...
+				ToolBarBuilder.AddComboButton(
+					FUIAction(FExecuteAction(), FCanExecuteAction::CreateSP(InEditorPanel, &SCurveEditorPanel::IsTangentModeComboEnabled)),
+					FOnGetContent::CreateSP(InEditorPanel, &SCurveEditorPanel::MakeTangentModeMenu),
+					TAttribute<FText>::CreateSP(&InEditorPanel.Get(), &SCurveEditorPanel::GetTangentModeLabel),
+					TAttribute<FText>::CreateSP(&InEditorPanel.Get(), &SCurveEditorPanel::GetTangentModeTooltip),
+					TAttribute<FSlateIcon>::CreateSP(&InEditorPanel.Get(), &SCurveEditorPanel::GetTangentModeIcon)
+					);
+				ToolBarBuilder.SetLabelVisibility(EVisibility::Collapsed); // ... but don't affect the other entries
+				ToolBarBuilder.EndStyleOverride();
+				
+				const int32 SupportedTangentTypes = InEditorPanel->CurveEditor->GetSupportedTangentTypes();
 				if (SupportedTangentTypes & (int32)ECurveEditorTangentTypes::InterpolationCubicWeighted)
 				{
 					ToolBarBuilder.AddToolBarButton(FCurveEditorCommands::Get().InterpolationToggleWeighted);
 				}
-
-				// We re-use key interpolation checks here, as you can set them under the same conditions.
-				FCanExecuteAction CanSetInfinities = FCanExecuteAction::CreateSP(InEditorPanel, &SCurveEditorPanel::CanSetKeyInterpolation);
-				
-				// Dropdown Menu for choosing your pre infinity for selected curves.
-				ToolBarBuilder.AddComboButton(
-					FUIAction(FExecuteAction(), CanSetInfinities),
-					FOnGetContent::CreateSP(InEditorPanel, &SCurveEditorPanel::MakeCurveExtrapolationMenu, false),
-					LOCTEXT("CurveEditorPreInfinityOptions", "Pre Infinity"),
-					LOCTEXT("CurveEditorPreInfinityOptionsToolTip", "Choose how the curve is evaluated if sampled before the first key"),
-					TAttribute<FSlateIcon>(InEditorPanel, &SCurveEditorPanel::GetCurveExtrapolationPreIcon)
-				);
-
-					// Dropdown menu for choosing your post infinity for selected curves
-				ToolBarBuilder.AddComboButton(
-					FUIAction(FExecuteAction(), CanSetInfinities),
-					FOnGetContent::CreateSP(InEditorPanel, &SCurveEditorPanel::MakeCurveExtrapolationMenu, true),
-					LOCTEXT("CurveEditorPostInfinityOptions", "Post Infinity"),
-					LOCTEXT("CurveEditorPostInfinityOptionsToolTip", "Choose how the curve is evaluated if sampled after the last key"),
-					TAttribute<FSlateIcon>(InEditorPanel, &SCurveEditorPanel::GetCurveExtrapolationPostIcon)
-				);
+				ToolBarBuilder.AddToolBarButton(FCurveEditorCommands::Get().FlattenTangents);
+				ToolBarBuilder.AddToolBarButton(FCurveEditorCommands::Get().StraightenTangents);
 			}
 			ToolBarBuilder.EndSection();
 
-			ToolBarBuilder.BeginSection("Filters");
+			ToolBarBuilder.BeginSection("Curves");
 			{
-				ToolBarBuilder.AddToolBarButton(FCurveEditorCommands::Get().FlattenTangents);
-				ToolBarBuilder.AddToolBarButton(FCurveEditorCommands::Get().StraightenTangents);
-				ToolBarBuilder.AddToolBarButton(FCurveEditorCommands::Get().OpenUserImplementableFilterWindow);
+				ToolBarBuilder.SetLabelVisibility(EVisibility::Visible); // Show label...
+				ToolBarBuilder.AddComboButton(
+					FUIAction(),
+					FOnGetContent::CreateSP(InEditorPanel, &SCurveEditorPanel::MakeCurvesMenu, InExtender),
+					LOCTEXT("CurveEditor.CurvesCombo.Label", "Curves"),
+					LOCTEXT("CurveEditor.CurvesCombo.ToolTip", "Curves"),
+					FSlateIcon(FAppStyle::Get().GetStyleSetName(), "GenericCurveEditor.Curves")
+					);
+				ToolBarBuilder.SetLabelVisibility(EVisibility::Collapsed); // ... but don't affect the other entries
+				
 				InToolbarPromotedFilters->AppendToBuilder(ToolBarBuilder);
+				ToolBarBuilder.AddToolBarButton(FCurveEditorCommands::Get().OpenUserImplementableFilterWindow);
 			}
-			ToolBarBuilder.EndSection();
+			
 			ToolBarBuilder.EndStyleOverride();
 		}
 
 	};
-
 	Extender->AddToolBarExtension(
 		"Asset",
 		EExtensionHook::After,
 		GetCommands(),
 		FToolBarExtensionDelegate::CreateStatic(
-			&Local::FillToolbar, KeyDetailsView.ToSharedRef(), SharedThis(this), GetCurveEditor()->GetToolbarPromotedFilters().ToSharedRef()
+			&Local::FillToolbar, KeyDetailsView.ToSharedRef(), SharedThis(this), GetCurveEditor()->GetToolbarPromotedFilters().ToSharedRef(), Extender
 			)
 	);
 
@@ -1475,7 +1573,6 @@ TSharedRef<SWidget> SCurveEditorPanel::MakeGridSpacingMenu()
 
 TSharedRef<SWidget> SCurveEditorPanel::MakeAxisSnapMenu()
 {
-	// This builds the dropdown menu when looking at the Curve View Options combobox.
 	FMenuBuilder MenuBuilder(true, GetCommands());
 
 	MenuBuilder.AddMenuEntry(FCurveEditorCommands::Get().SetAxisSnappingNone);
@@ -1485,18 +1582,6 @@ TSharedRef<SWidget> SCurveEditorPanel::MakeAxisSnapMenu()
 	return MenuBuilder.MakeWidget();
 }
 
-TSharedRef<SWidget> SCurveEditorPanel::MakeViewModeMenu()
-{
-	// This builds the dropdown menu when looking at the Curve View Modes
-	FMenuBuilder MenuBuilder(true, GetCommands());
-
-	MenuBuilder.AddMenuEntry(FCurveEditorCommands::Get().SetViewModeAbsolute);
-	MenuBuilder.AddMenuEntry(FCurveEditorCommands::Get().SetViewModeStacked);
-	MenuBuilder.AddMenuEntry(FCurveEditorCommands::Get().SetViewModeNormalized);
-
-	return MenuBuilder.MakeWidget();
-}
-
 bool SCurveEditorPanel::IsInlineEditPanelEditable() const
 {
 	return CurveEditor->GetSelection().Count() > 0;
@@ -1509,6 +1594,35 @@ EVisibility SCurveEditorPanel::ShouldInstructionOverlayBeVisible() const
 	return bCurvesAreVisible ? EVisibility::Hidden : EVisibility::HitTestInvisible;
 }
 
+TSharedRef<SWidget> SCurveEditorPanel::MakeToolsComboMenu(TSharedPtr<FExtender> InExtender)
+{
+	FMenuBuilder MenuBuilder(true, CurveEditor->GetCommands(), InExtender);
+
+	MenuBuilder.BeginSection(TEXT("Tools"));
+	MenuBuilder.AddMenuEntry(FCurveEditorCommands::Get().DeactivateCurrentTool);
+	MenuBuilder.EndSection();
+
+	return MenuBuilder.MakeWidget();
+}
+
+FText SCurveEditorPanel::GetCurrentToolLabel() const
+{
+	const ICurveEditorToolExtension* CurrentTool = GetCurveEditor()->GetCurrentTool();
+	return CurrentTool ? CurrentTool->GetLabel() : FCurveEditorCommands::Get().DeactivateCurrentTool->GetLabel();
+}
+
+FText SCurveEditorPanel::GetCurrentToolDescription() const
+{
+	const ICurveEditorToolExtension* CurrentTool = GetCurveEditor()->GetCurrentTool();
+	return CurrentTool ? CurrentTool->GetDescription() : FCurveEditorCommands::Get().DeactivateCurrentTool->GetDescription();
+}
+
+FSlateIcon SCurveEditorPanel::GetCurrentToolIcon() const
+{
+	const ICurveEditorToolExtension* CurrentTool = GetCurveEditor()->GetCurrentTool();
+	return CurrentTool ? CurrentTool->GetIcon() : FCurveEditorCommands::Get().DeactivateCurrentTool->GetIcon();
+}
+
 void SCurveEditorPanel::OnSplitterFinishedResizing()
 {
 	SSplitter::FSlot const& LeftSplitterSlot = TreeViewSplitter->SlotAt(0);
diff --git a/Engine/Source/Editor/CurveEditor/Private/SCurveKeyDetailPanel.cpp b/Engine/Source/Editor/CurveEditor/Private/SCurveKeyDetailPanel.cpp
--- a/Engine/Source/Editor/CurveEditor/Private/SCurveKeyDetailPanel.cpp
+++ b/Engine/Source/Editor/CurveEditor/Private/SCurveKeyDetailPanel.cpp
@@ -19,6 +19,7 @@
 #include "SlotBase.h"
 #include "Types/SlateEnums.h"
 #include "Widgets/Input/SEditableTextBox.h"
+#include "Widgets/Layout/SBox.h"
 #include "Widgets/SBoxPanel.h"
 #include "Widgets/SWidget.h"
 
@@ -183,43 +184,29 @@ void SCurveKeyDetailPanel::ConstructChildLayout(TSharedPtr<SWidget> TimeWidget,
 
 	ChildSlot
 	[
-		SNew(SHorizontalBox)
-		// "Time" Label
-		// + SHorizontalBox::Slot()
-		// .VAlign(VAlign_Center)
-		// .AutoWidth()
-		// [
-		// 	SNew(STextBlock)
-		// 	.Text(LOCTEXT("KeyDetailTimeLabel", "Time"))
-		// ]
-
-		// "Time" Edit box
-		+ SHorizontalBox::Slot()
-		.VAlign(VAlign_Fill)
-		// .HAlign(HAlign_Left)
-		.Padding(4.f, 0.f, 0.f, 2.f)
-		.FillWidth(0.5f)
+		SNew(SBox)
+		.MaxDesiredWidth(FAppStyle::Get().GetFloat("CurveEditor.KeyDetailWidth"))
 		[
-			TimeWidget.ToSharedRef()
-		]
+			SNew(SHorizontalBox)
 
-		// "Value" Label
-		// + SHorizontalBox::Slot()
-		// .VAlign(VAlign_Center)
-		// .AutoWidth()
-		// [
-		// 	SNew(STextBlock)
-		// 	.Text(LOCTEXT("KeyDetailValueLabel", "Value"))
-		// ]
-
-		// "Value" Edit box
-		+ SHorizontalBox::Slot()
-		.VAlign(VAlign_Fill)
-		// .HAlign(HAlign_Left)
-		.FillWidth(0.5f)
-		.Padding(4.f, 0.f, 0.f, 2.f)
-		[
-			ValueWidget.ToSharedRef()
+			// "Time" Edit box
+			+ SHorizontalBox::Slot()
+			.VAlign(VAlign_Center)
+			.Padding(4.f, 0.f, 0.f, 2.f)
+			.FillWidth(0.5f)
+			[
+				TimeWidget.ToSharedRef()
+			]
+
+
+			// "Value" Edit box
+			+ SHorizontalBox::Slot()
+			.VAlign(VAlign_Center)
+			.FillWidth(0.5f)
+			.Padding(4.f, 0.f, 0.f, 2.f)
+			[
+				ValueWidget.ToSharedRef()
+			]
 		]
 	];
 }
diff --git a/Engine/Source/Editor/CurveEditor/Public/ICurveEditorToolExtension.h b/Engine/Source/Editor/CurveEditor/Public/ICurveEditorToolExtension.h
--- a/Engine/Source/Editor/CurveEditor/Public/ICurveEditorToolExtension.h
+++ b/Engine/Source/Editor/CurveEditor/Public/ICurveEditorToolExtension.h
@@ -9,11 +9,13 @@ class FUICommandList;
 class FPaintArgs;
 class FSlateRect;
 class FSlateWindowElementList;
+class FText;
 class FWidgetStyle;
-struct FPointerEvent;
 class SWidget;
 
 struct FGeometry;
+struct FPointerEvent;
+struct FSlateIcon;
 
 /**
 * A unique identifier for a tool in the curve editor
@@ -134,6 +136,13 @@ public:
 		ToolID = InToolID;
 	}
 
+	/** @return Label for this tool. Displayed by toolbar combo button. */
+	virtual FText GetLabel() const = 0;
+	/** @return Description for this tool. Displayed by toolbar combo button. */
+	virtual FText GetDescription() const = 0;
+	/** @return Icon for this tool. Displayed by toolbar combo button. */
+	virtual FSlateIcon GetIcon() const = 0;
+
 	FOnOptionsRefresh OnOptionsRefreshDelegate;
 
 protected:
diff --git a/Engine/Source/Editor/CurveEditor/Public/SCurveEditorPanel.h b/Engine/Source/Editor/CurveEditor/Public/SCurveEditorPanel.h
--- a/Engine/Source/Editor/CurveEditor/Public/SCurveEditorPanel.h
+++ b/Engine/Source/Editor/CurveEditor/Public/SCurveEditorPanel.h
@@ -32,6 +32,7 @@
 class FCurveEditor;
 class FExtender;
 class FTabManager;
+class FMenuBuilder;
 class FUICommandList;
 class IDetailsView;
 class IGraphEditorView;
@@ -197,11 +198,43 @@ private:
 
 	TSharedRef<SWidget> MakeGridSpacingMenu();
 	TSharedRef<SWidget> MakeAxisSnapMenu();
-	TSharedRef<SWidget> MakeViewModeMenu();
 
 	bool IsInlineEditPanelEditable() const;
 	EVisibility ShouldInstructionOverlayBeVisible() const;
 
+private:
+	
+	TSharedRef<SWidget> MakeToolsComboMenu(TSharedPtr<FExtender> InExtender);
+	FText GetCurrentToolLabel() const;
+	FText GetCurrentToolDescription() const;
+	FSlateIcon GetCurrentToolIcon() const;
+	
+private:
+	
+	// Tangent mode toolbar combo button
+	TSharedRef<SWidget> MakeTangentModeMenu();
+	FText GetTangentModeLabel() const;
+	FText GetTangentModeTooltip() const;
+	FSlateIcon GetTangentModeIcon() const;
+	bool IsTangentModeComboEnabled() const;
+	
+	enum class ETangentModeComboState : uint8
+	{
+		NoSelection,
+
+		// All keys in selection share the following mode:
+		Constant,
+		Linear,
+		CubicAuto,
+		CubicSmartAuto,
+		CubicUser,
+		CubicBreak,
+
+		// The keys are mixed
+		Mixed
+	};
+	ETangentModeComboState DetermineTangentMode() const;
+	
 private:
 
 	/**
@@ -215,21 +248,24 @@ private:
 	void UpdateCommonCurveInfo();
 
 	/** Creates the drop-down list you see when changing Curve View options. */
-	TSharedRef<SWidget> MakeCurveEditorCurveViewOptionsMenu();
-	TSharedRef<SWidget> MakeCurveExtrapolationMenu(const bool bInPostExtrapolation);
+	TSharedRef<SWidget> MakeCurveEditorCurveViewOptionsMenu(TSharedPtr<FExtender> InExtender);
+	/** Creates the drop-down for changing curves on a whole. */
+	TSharedRef<SWidget> MakeCurvesMenu(TSharedPtr<FExtender> InExtender);
+	void AddPreInfinityToMenu(FMenuBuilder& InMenuBuilder);
+	void AddPostInfinityToMenu(FMenuBuilder& InMenuBuilder);
 	FSlateIcon GetCurveExtrapolationPreIcon() const;
 	FSlateIcon GetCurveExtrapolationPostIcon() const;
 
 	/** Creates the Curve Editor Filter UI and pre-populates it with the specified class. */
 	void ShowCurveFilterUI(TSubclassOf<UCurveEditorFilterBase> FilterClass);
-
+	
 private:
 
 	/**
 	 * Bind command mappings for this widget
 	 */
 	void BindCommands();
-
+	
 	/**
 	 * Assign new attributes to the currently selected keys
 	 */
diff --git a/Engine/Source/Editor/EditorStyle/Private/StarshipStyle.cpp b/Engine/Source/Editor/EditorStyle/Private/StarshipStyle.cpp
--- a/Engine/Source/Editor/EditorStyle/Private/StarshipStyle.cpp
+++ b/Engine/Source/Editor/EditorStyle/Private/StarshipStyle.cpp
@@ -7411,7 +7411,7 @@ void FStarshipEditorStyle::FStyle::SetupToolkitStyles()
 	// Curve Editor
 	{
 		Set("CurveAssetEditor.Tabs.Properties", new IMAGE_BRUSH("Icons/AssetIcons/CurveBase_16x", Icon16x16));
-
+		
 		Set("CurveEditor.FitHorizontally", new IMAGE_BRUSH("Icons/icon_CurveEditor_Horizontal_40x", Icon40x40));
 		Set("CurveEditor.FitVertically", new IMAGE_BRUSH("Icons/icon_CurveEditor_Vertical_40x", Icon40x40));
 		Set("CurveEditor.Fit", new IMAGE_BRUSH("Icons/icon_CurveEditor_ZoomToFit_40x", Icon40x40));
@@ -7468,10 +7468,19 @@ void FStarshipEditorStyle::FStyle::SetupToolkitStyles()
 		CurveEditorToolbar.SetSeparatorPadding(    FMargin(2.0f, 0.0f));
 
 		Set("CurveEditorToolBar", CurveEditorToolbar);
-		
-		CurveEditorToolbar.SetComboLabelMinWidth(95.f);
+
+		// Some min sizes for toolbar combo buttons. We don't want the combo buttons to resize depending on the size of their contained label.
 		CurveEditorToolbar.SetComboContentHorizontalAlignment(HAlign_Center);
-		Set("CurveEditorTweenToolbar", CurveEditorToolbar);
+		{
+			CurveEditorToolbar.SetComboLabelMinWidth(68.f);
+			Set("CurveEditorToolBar.ToolsCombo", CurveEditorToolbar);
+			CurveEditorToolbar.SetComboLabelMinWidth(90.f);
+			Set("CurveEditorToolBar.TangentMode", CurveEditorToolbar);
+			CurveEditorToolbar.SetComboLabelMinWidth(90.f);
+			Set("CurveEditorTweenToolbar", CurveEditorToolbar);
+		}
+
+		Set("CurveEditor.KeyDetailWidth", 180.f);
 
 		// Tab
 		Set("GenericCurveEditor", new IMAGE_BRUSH_SVG("Starship/GenericCurveEditor/CurveEditor", Icon20x20));
@@ -7492,6 +7501,8 @@ void FStarshipEditorStyle::FStyle::SetupToolkitStyles()
 		Set("GenericCurveEditor.InterpolationToggleWeighted", new IMAGE_BRUSH_SVG("Starship/GenericCurveEditor/TangentsWeighted", Icon20x20));
 		Set("GenericCurveEditor.InterpolationLinear", new IMAGE_BRUSH_SVG("Starship/GenericCurveEditor/TangentsLinear", Icon20x20));
 		Set("GenericCurveEditor.InterpolationConstant", new IMAGE_BRUSH_SVG("Starship/GenericCurveEditor/TangentsConstant", Icon20x20));
+		Set("GenericCurveEditor.InterpolationMixed", new IMAGE_BRUSH_SVG("Starship/GenericCurveEditor/TangentsMixed", Icon20x20));
+		Set("GenericCurveEditor.InterpolationNoSelection", new IMAGE_BRUSH_SVG("Starship/GenericCurveEditor/TangentsNoSelection", Icon20x20));
 
 		// Tangent Modifications
 		Set("GenericCurveEditor.FlattenTangents", new IMAGE_BRUSH_SVG("Starship/GenericCurveEditor/MiscFlatten", Icon20x20));
@@ -7541,6 +7552,8 @@ void FStarshipEditorStyle::FStyle::SetupToolkitStyles()
 		Set("GenericCurveEditor.Pin_Inactive", new IMAGE_BRUSH("Common/PushPin_Up", Icon16x16));
 		
 		Set("GenericCurveEditor.Select", new IMAGE_BRUSH("GenericCurveEditor/Icons/SelectButton", Icon16x16));
+		
+		Set("GenericCurveEditor.Curves", new IMAGE_BRUSH_SVG("Starship/GenericCurveEditor/GraphViewAbsolute_20", Icon20x20)); 
 	}
 
 	// Generic Curve Editor Tools
