Curve Editor: Smart key snapping. You can now right-click a key selection, and snap keys to full frames. The keys will be placed where the curve intersects with the frames.

diff --git a/Engine/Plugins/Editor/CurveEditorTools/Source/CurveEditorTools/Private/Tools/Lattice/CurveEditorLatticeTool.cpp b/Engine/Plugins/Editor/CurveEditorTools/Source/CurveEditorTools/Private/Tools/Lattice/CurveEditorLatticeTool.cpp
--- a/Engine/Plugins/Editor/CurveEditorTools/Source/CurveEditorTools/Private/Tools/Lattice/CurveEditorLatticeTool.cpp
+++ b/Engine/Plugins/Editor/CurveEditorTools/Source/CurveEditorTools/Private/Tools/Lattice/CurveEditorLatticeTool.cpp
@@ -4,7 +4,10 @@
 
 #include "CurveEditor.h"
 #include "CurveEditorToolCommands.h"
+#include "CoreGlobals.h"
 #include "Drag/LatticeDragOp_MoveControlPoints.h"
+#include "Editor.h"
+#include "Editor/TransBuffer.h"
 #include "Internationalization/Text.h"
 #include "Mirror/LatticeEdgeTangentMirrorOp.h"
 #include "Mirror/LatticePointTangentMirrorOp.h"
@@ -54,6 +57,19 @@ FCurveEditorLatticeTool::FCurveEditorLatticeTool(TWeakPtr<FCurveEditor> InCurveE
 			UndoObject->RestoreEditorState(*Editor);
 		}
 	});
+	
+	if (UTransBuffer* Transactor = GEditor ? Cast<UTransBuffer>(GEditor->Trans) : nullptr)
+	{
+		Transactor->OnTransactionStateChanged().AddRaw(this, &FCurveEditorLatticeTool::OnTransactionStateChanged);
+	}
+}
+
+FCurveEditorLatticeTool::~FCurveEditorLatticeTool()
+{
+	if (UTransBuffer* Transactor = GEditor ? Cast<UTransBuffer>(GEditor->Trans) : nullptr)
+	{
+		Transactor->OnTransactionStateChanged().RemoveAll(this);
+	}
 }
 
 void FCurveEditorLatticeTool::Tick(const FGeometry& InViewContainerGeometry, const double InCurrentTime, const float InDeltaTime)
@@ -68,8 +84,10 @@ void FCurveEditorLatticeTool::Tick(const FGeometry& InViewContainerGeometry, con
 	
 	if (LatticeDeformerState && LatticeDeformerState->bHasRequestedRefresh)
 	{
+		const ELatticeUpdateFlags Flags = LatticeDeformerState->RefreshFlags;
+		LatticeDeformerState->RefreshFlags = ELatticeUpdateFlags::None;
 		LatticeDeformerState->bHasRequestedRefresh = false;
-		RecomputeLatticeGrid(CurveEditor.ToSharedRef());
+		RecomputeLatticeGrid(CurveEditor.ToSharedRef(), Flags);
 	}
 	else
 	{
@@ -269,7 +287,7 @@ FCurveEditorLatticeTool::FLatticeDeformerState::FLatticeDeformerState(const FLat
 	, CurveChangeListener(InCurveEditor)
 {}
 
-void FCurveEditorLatticeTool::UpdateDeformerBounds(const TSharedRef<FCurveEditor>& InCurveEditor)
+void FCurveEditorLatticeTool::UpdateDeformerBounds(const TSharedRef<FCurveEditor>& InCurveEditor, ELatticeUpdateFlags InRefreshFlags)
 {
 	if (SelectionSerialNumber != InCurveEditor->Selection.GetSerialNumber())
 	{
@@ -294,7 +312,10 @@ void FCurveEditorLatticeTool::UpdateDeformerBounds(const TSharedRef<FCurveEditor
 		AddSelectionToDeformer(GetDeformer(), *InCurveEditor);
 		LatticeDeformerState->CurveChangeListener.OnCurveModified().AddRaw(this, &FCurveEditorLatticeTool::OnCurvesModified);
 
-		FScopedTransaction Transaction(TEXT("CurveEditorLatticeTool"), LOCTEXT("SetSelection", "Set lattice selection"), nullptr);
+		const bool bActuallyTransact = !EnumHasAnyFlags(InRefreshFlags, ELatticeUpdateFlags::NoTransaction)
+			// Re-use ongoing transaction if currently recording into buffer
+			&& GUndo == nullptr && !GIsTransacting; 
+		FScopedTransaction Transaction(TEXT("CurveEditorLatticeTool"), LOCTEXT("SetSelection", "Set lattice selection"), nullptr, bActuallyTransact);
 		UndoObject->Modify();
 		UndoObject->SaveState(*InCurveEditor, GetDeformer());
 	}
@@ -303,10 +324,34 @@ void FCurveEditorLatticeTool::UpdateDeformerBounds(const TSharedRef<FCurveEditor
 void FCurveEditorLatticeTool::OnCurvesModified()
 {
 	// The modification was driven by us?
-	LatticeDeformerState->bHasRequestedRefresh = !LatticeDeformerState->bIsModifyingCurves;
+	if (LatticeDeformerState->bIsModifyingCurves)
+	{
+		return;
+	}
+
+	// If an external change happens during a transaction or undo / redo, and we refresh the grid shape next tick, do not generate another transaction.
+	// Instead of next tick, we should defer the change to the end of the current transaction but that is a bit tricky.
+	// Instead of implementing the deferring TODO UE-242452: the lattice tool should stop tracking key selection and we should make core Curve Editor track the key selection with transactions. 
+	if (GUndo || GIsTransacting)
+	{
+		LatticeDeformerState->RefreshFlags = ELatticeUpdateFlags::NoTransaction;
+	}
+	LatticeDeformerState->bHasRequestedRefresh = true;
 }
 
-void FCurveEditorLatticeTool::RecomputeLatticeGrid(const TSharedRef<FCurveEditor>& InCurveEditor)
+void FCurveEditorLatticeTool::OnTransactionStateChanged(
+	const FTransactionContext&, ETransactionStateEventType TransactionStateEvent
+	)
+{
+	// Technically, we should check whether the transaction actually affected the lattice... but this isn't computationally expensive so whatever.
+	const bool bNeedsRecompute_UndoRedo = TransactionStateEvent == ETransactionStateEventType::UndoRedoFinalized;
+	if (bNeedsRecompute_UndoRedo)
+	{
+		RecomputeLatticeGrid();
+	}
+}
+
+void FCurveEditorLatticeTool::RecomputeLatticeGrid(const TSharedRef<FCurveEditor>& InCurveEditor, ELatticeUpdateFlags InRefreshFlags)
 {
 	checkSlow(LatticeDeformerState);
 	
@@ -314,7 +359,7 @@ void FCurveEditorLatticeTool::RecomputeLatticeGrid(const TSharedRef<FCurveEditor
 	// In the past, we tried to preserve the existing shape by just moving it over the moved keys but there were too many different tools interacting
 	// which made this difficult to get right in every case. Hence, we decided just to recompute the grid on move.
 	LatticeDeformerState.Reset();
-	UpdateDeformerBounds(InCurveEditor);
+	UpdateDeformerBounds(InCurveEditor, InRefreshFlags);
 }
 
 void FCurveEditorLatticeTool::CancelDrag()
@@ -451,6 +496,7 @@ void FCurveEditorLatticeTool::StartDragDrop(const FLatticeBounds& InBounds, cons
 void FCurveEditorLatticeTool::OnMoveControlPoint(TConstArrayView<int32> InVertexIndices, TConstArrayView<FVector2D> InNewControlPointPositions, EPropertyChangeType::Type InChangeType)
 {
 	check(InVertexIndices.Num() == 1 && InNewControlPointPositions.Num() == 1);
+	TGuardValue Guard(LatticeDeformerState->bIsModifyingCurves, true);
 	
 	if (DelayedDragData->PointDragTangentMirrorOperation)
 	{
@@ -463,6 +509,7 @@ void FCurveEditorLatticeTool::OnMoveControlPoint(TConstArrayView<int32> InVertex
 void FCurveEditorLatticeTool::OnMoveEdge(TConstArrayView<int32> InVertexIndices, TConstArrayView<FVector2D> InNewControlPointPositions, EPropertyChangeType::Type InChangeType)
 {
 	check(InVertexIndices.Num() == 2 && InNewControlPointPositions.Num() == 2);
+	TGuardValue Guard(LatticeDeformerState->bIsModifyingCurves, true);
 	
 	if (DelayedDragData->EdgeDragTangentMirrorOperation)
 	{
@@ -483,6 +530,8 @@ void FCurveEditorLatticeTool::OnSetControlPoints(
 		return;
 	}
 	
+	TGuardValue Guard(LatticeDeformerState->bIsModifyingCurves, true);
+	
 	// Combining Handles & Keys into large arrays like this is significantly faster than feeding each key into SetKeyPositions separately.
 	// For 9000 keys it goes from 2136 ms to 238 ms.
 	struct FCurveChanges
@@ -521,7 +570,6 @@ void FCurveEditorLatticeTool::OnSetControlPoints(
 	// - SetKeyPositions 138.954 ms		- Total 139.229 ms
 	// - SetKeyPositions 269.523 ms		- Total 269.811 ms
 	// Conclusion: The following SetKeyPositions block is slow.
-	TGuardValue Guard(LatticeDeformerState->bIsModifyingCurves, true);
 	for (auto It = Changes.CreateConstIterator(); It; ++It)
 	{
 		const FCurveChanges& Pair = It->Value;
@@ -598,7 +646,7 @@ void FCurveEditorLatticeTool::MirrorGrid()
 	LatticeDeformerState.Reset();
 	UpdateDeformerBounds(CurveEditor.ToSharedRef());
 }
-	
+
 template <typename TCallback> requires std::is_invocable_r_v<bool, TCallback, const FLatticeHoverState&, const FLatticeBounds&>
 bool FCurveEditorLatticeTool::ComputeHoverStateForLeftClick(
 	const FGeometry& InViewContainerGeometry, const FPointerEvent& InMouseEvent, TCallback&& InCallback
diff --git a/Engine/Plugins/Editor/CurveEditorTools/Source/CurveEditorTools/Private/Tools/Lattice/CurveEditorLatticeTool.h b/Engine/Plugins/Editor/CurveEditorTools/Source/CurveEditorTools/Private/Tools/Lattice/CurveEditorLatticeTool.h
--- a/Engine/Plugins/Editor/CurveEditorTools/Source/CurveEditorTools/Private/Tools/Lattice/CurveEditorLatticeTool.h
+++ b/Engine/Plugins/Editor/CurveEditorTools/Source/CurveEditorTools/Private/Tools/Lattice/CurveEditorLatticeTool.h
@@ -10,6 +10,7 @@
 #include "Mirror/LatticeEdgeTangentMirrorOp.h"
 #include "Mirror/LatticePointTangentMirrorOp.h"
 #include "Misc/CurveChangeListener.h"
+#include "Misc/EnumClassFlags.h"
 #include "Misc/LatticeDrawUtils.h"
 #include "ScopedTransaction.h"
 #include "Templates/SharedPointer.h"
@@ -18,10 +19,20 @@
 class FCurveEditor;
 class FText;
 class UCurveEditorTools_LatticeUndoData;
+struct FTransactionContext;
+enum class ETransactionStateEventType : uint8;
 
 namespace UE::CurveEditorTools
 {
 struct FLatticeBounds;
+
+enum class ELatticeUpdateFlags : uint8
+{
+	None,
+	/** Don't generate any transaction for this update. */
+	NoTransaction = 1 << 1
+};
+ENUM_CLASS_FLAGS(ELatticeUpdateFlags);
 	
 /**
  * Places a simple lattice deformer over the selection.
@@ -32,6 +43,7 @@ class FCurveEditorLatticeTool : public ICurveEditorToolExtension, public FGCObje
 public:
 	
 	explicit FCurveEditorLatticeTool(TWeakPtr<FCurveEditor> InCurveEditor);
+	virtual ~FCurveEditorLatticeTool() override;
 
 	//~ Begin ICurveEditorToolExtension Interface
 	virtual void Tick(const FGeometry& InViewContainerGeometry, const double InCurrentTime, const float InDeltaTime) override;
@@ -75,6 +87,8 @@ private:
 		bool bIsModifyingCurves = false;
 		/** If true, the curves have been modified and the lattice overlay should be recomputed next frame. */
 		bool bHasRequestedRefresh = false;
+		/** The flags to use when bHasRequestedRefresh == true. */
+		ELatticeUpdateFlags RefreshFlags = ELatticeUpdateFlags::None;
 
 		explicit FLatticeDeformerState(const FLatticeBounds& Lattice, const TSharedRef<FCurveEditor>& InCurveEditor);
 	};
@@ -133,12 +147,16 @@ private:
 	TObjectPtr<UCurveEditorTools_LatticeUndoData> UndoObject;
 	
 	/** Creates or destroys LatticeDeformer based on the selection. */
-	void UpdateDeformerBounds(const TSharedRef<FCurveEditor>& InCurveEditor);
+	void UpdateDeformerBounds(const TSharedRef<FCurveEditor>& InCurveEditor, ELatticeUpdateFlags InRefreshFlags = ELatticeUpdateFlags::None);
 
 	/** Updates the bounds after the user has externally moved the keys, e.g. MMB move keys, etc. */
 	void OnCurvesModified();
+	/**  */
+	void OnTransactionStateChanged(const FTransactionContext& TransactionContext, ETransactionStateEventType TransactionStateEvent);
+	
 	/** Call this if the keys have externally moved. Tries to fit the lattice grid over the keys or reconstructs a new rectangular one if not possible. */
-	void RecomputeLatticeGrid(const TSharedRef<FCurveEditor>& InCurveEditor);
+	void RecomputeLatticeGrid(const TSharedRef<FCurveEditor>& InCurveEditor, ELatticeUpdateFlags InRefreshFlags = ELatticeUpdateFlags::None);
+	void RecomputeLatticeGrid(ELatticeUpdateFlags InRefreshFlags = ELatticeUpdateFlags::None);
 	
 	/** Stops the active drag operation. */
 	void CancelDrag();
@@ -172,4 +190,12 @@ private:
 	template<typename TCallback> requires std::is_invocable_r_v<bool, TCallback, const FLatticeHoverState&, const FLatticeBounds&>
 	bool ComputeHoverStateForLeftClick(const FGeometry& InViewContainerGeometry, const FPointerEvent& InMouseEvent, TCallback&& InCallback);
 };
+
+inline void FCurveEditorLatticeTool::RecomputeLatticeGrid(ELatticeUpdateFlags InRefreshFlags)
+{
+	if (const TSharedPtr<FCurveEditor> CurveEditorPin = WeakCurveEditor.Pin())
+	{
+		RecomputeLatticeGrid(CurveEditorPin.ToSharedRef(), InRefreshFlags);
+	}
+}
 }
diff --git a/Engine/Source/Editor/CurveEditor/Private/CurveEditor.cpp b/Engine/Source/Editor/CurveEditor/Private/CurveEditor.cpp
--- a/Engine/Source/Editor/CurveEditor/Private/CurveEditor.cpp
+++ b/Engine/Source/Editor/CurveEditor/Private/CurveEditor.cpp
@@ -55,6 +55,7 @@
 #include "UObject/UObjectGlobals.h"
 #include "UObject/UnrealNames.h"
 #include "UnrealExporter.h"
+#include "Misc/SmartSnap.h"
 #include "Widgets/Colors/SColorPicker.h"
 #include "Widgets/Notifications/SNotificationList.h"
 
@@ -484,6 +485,8 @@ void FCurveEditor::BindCommands()
 		CommandList->MapAction(FCurveEditorCommands::Get().StraightenTangents, FExecuteAction::CreateSP(this, &FCurveEditor::StraightenSelection), FCanExecuteAction::CreateSP(this, &FCurveEditor::CanFlattenOrStraightenSelection) );
 	}
 
+	CommandList->MapAction(FCurveEditorCommands::Get().SmartSnapKeys, FExecuteAction::CreateSP(this, &FCurveEditor::SmartSnapSelection), FCanExecuteAction::CreateSP(this, &FCurveEditor::CanSmartSnapSelection));
+
 	// Curve Colors
 	{
 		CommandList->MapAction(FCurveEditorCommands::Get().SetRandomCurveColorsForSelected, FExecuteAction::CreateSP(this, &FCurveEditor::SetRandomCurveColorsForSelected), FCanExecuteAction());
@@ -2061,6 +2064,35 @@ bool FCurveEditor::CanFlattenOrStraightenSelection() const
 	return Selection.Count() > 0;
 }
 
+void FCurveEditor::SmartSnapSelection()
+{
+	FScopedTransaction Transaction(LOCTEXT("SmartSnapKeys", "Smart Snap"));
+	
+	TMap<FCurveModelID, FKeyHandleSet> OutKeysToSelect;
+	UE::CurveEditor::EnumerateSmartSnappableKeys(*this, Selection.GetAll(), OutKeysToSelect,
+		[](const FCurveModelID&, FCurveModel& CurveModel, const UE::CurveEditor::FSmartSnapResult& SnapResult)
+		{
+			CurveModel.Modify();
+			ApplySmartSnap(CurveModel, SnapResult);
+		});
+
+	// Some keys might have been removed. Clean up the selection.
+	Selection.Clear();
+	for (const TTuple<FCurveModelID, FKeyHandleSet>& OutSet : OutKeysToSelect)
+	{
+		Selection.Add(OutSet.Key, ECurvePointType::Key, OutSet.Value.AsArray());
+	}
+}
+
+bool FCurveEditor::CanSmartSnapSelection() const
+{
+	return UE::CurveEditor::CanSmartSnapSelection(Selection);
+}
+
+void FCurveEditor::UpdateGeometry(const FGeometry& CurrentGeometry)
+{
+}
+
 void FCurveEditor::SetRandomCurveColorsForSelected()
 {
 	TSet<FCurveModelID> CurveModelIDs = GetSelectionFromTreeAndKeys();
diff --git a/Engine/Source/Editor/CurveEditor/Private/CurveEditorCommands.cpp b/Engine/Source/Editor/CurveEditor/Private/CurveEditorCommands.cpp
--- a/Engine/Source/Editor/CurveEditor/Private/CurveEditorCommands.cpp
+++ b/Engine/Source/Editor/CurveEditor/Private/CurveEditorCommands.cpp
@@ -32,6 +32,8 @@ void FCurveEditorCommands::RegisterCommands()
 
 	UI_COMMAND(FlattenTangents, "Flatten", "Flatten Tangents", EUserInterfaceActionType::Button, FInputChord(EKeys::Six));
 	UI_COMMAND(StraightenTangents, "Straighten", "Straighten Tangents", EUserInterfaceActionType::Button, FInputChord());
+	
+	UI_COMMAND(SmartSnapKeys, "Smart Snap", "Snaps selected keys to the closest whole frame.\nThe key is placed where the curve intersects with the frame (imagine vertical line going down from the frame).", EUserInterfaceActionType::Button, FInputChord(EKeys::Seven));
 
 	UI_COMMAND(MatchLastTangentToFirst, "Match Last Tangent To First", "Match last keys tangents to the first keys tangents", EUserInterfaceActionType::Button, FInputChord());
 	UI_COMMAND(MatchFirstTangentToLast, "Match First Tangent To Last", "Match first keys tangents to the last keys tangents", EUserInterfaceActionType::Button, FInputChord());
diff --git a/Engine/Source/Editor/CurveEditor/Private/CurveEditorContextMenu.cpp b/Engine/Source/Editor/CurveEditor/Private/CurveEditorContextMenu.cpp
--- a/Engine/Source/Editor/CurveEditor/Private/CurveEditorContextMenu.cpp
+++ b/Engine/Source/Editor/CurveEditor/Private/CurveEditorContextMenu.cpp
@@ -24,101 +24,225 @@ class IBufferedCurveModel;
 
 #define LOCTEXT_NAMESPACE "CurveEditorContextMenu"
 
-void FCurveEditorContextMenu::BuildMenu(FMenuBuilder& MenuBuilder, TSharedRef<FCurveEditor> CurveEditor, TOptional<FCurvePointHandle> ClickedPoint, TOptional<FCurveModelID> HoveredCurveID)
+namespace UE::CurveEditor::ContextDetail
 {
-	int32 NumSelectedKeys = CurveEditor->GetSelection().Count();
-
-	TSharedPtr<FCurveEditor> LocalCurveEditor = CurveEditor->AsShared();
-
-	int32 NumActiveCurves = LocalCurveEditor->GetSelectionFromTreeAndKeys().Num();
-
-	int32 NumBufferedCurves = 0;
-	for (const TUniquePtr<IBufferedCurveModel>& BufferedCurve : LocalCurveEditor->GetBufferedCurves())
-	{
-		if (LocalCurveEditor->IsActiveBufferedCurve(BufferedCurve))
-		{
-			++NumBufferedCurves;
-		}
-	}
-
+/** Holds shared args to make calling maker functions easier. */
+class FContextMenuBuilder
+{
+public:
+	
+	const TSharedRef<FCurveEditor> CurveEditor;
+	const TOptional<FCurvePointHandle> ClickedPoint;
+	const TOptional<FCurveModelID> HoveredCurveID;
+	
+	const int32 NumSelectedKeys = CurveEditor->GetSelection().Count();
+	const TSharedPtr<FCurveEditor> LocalCurveEditor = CurveEditor->AsShared();
+	const int32 NumActiveCurves = LocalCurveEditor->GetSelectionFromTreeAndKeys().Num();
+	const int32 NumBufferedCurves = 0;
+	
 	// We change the name to reflect the current number of curves selected.
-	TAttribute<FText> BufferedCurvesText = FText::Format(LOCTEXT("BufferCurvesContextMenu", "Buffer {0} Curves"), NumActiveCurves);
-	TAttribute<FText> SwapBufferedCurvesText = FText::Format(LOCTEXT("SwapBufferedCurvesContextMenu", "Swap {0} Buffered Curves onto {1} Selected Curves"), NumBufferedCurves, NumActiveCurves);
-	TAttribute<FText> ApplyBufferedCurvesText = FText::Format(LOCTEXT("ApplyBufferedCurvesContextMenu", "Apply {0} Buffered Curves onto {1} Selected Curves"), NumBufferedCurves, NumActiveCurves);
+	const TAttribute<FText> BufferedCurvesText
+		= FText::Format(LOCTEXT("BufferCurvesContextMenu", "Buffer {0} Curves"), NumActiveCurves);
+	const TAttribute<FText> SwapBufferedCurvesText
+		= FText::Format(LOCTEXT("SwapBufferedCurvesContextMenu", "Swap {0} Buffered Curves onto {1} Selected Curves"), NumBufferedCurves, NumActiveCurves);
+	const TAttribute<FText> ApplyBufferedCurvesText
+		= FText::Format(LOCTEXT("ApplyBufferedCurvesContextMenu", "Apply {0} Buffered Curves onto {1} Selected Curves"), NumBufferedCurves, NumActiveCurves);
 
 	const FCurveModel* HoveredCurve = HoveredCurveID.IsSet() ? CurveEditor->FindCurve(HoveredCurveID.GetValue()) : nullptr;
+	
+	FContextMenuBuilder(TSharedRef<FCurveEditor> CurveEditor, TOptional<FCurvePointHandle> ClickedPoint, TOptional<FCurveModelID> HoveredCurveID)
+		: CurveEditor(CurveEditor)
+		, ClickedPoint(ClickedPoint)
+		, HoveredCurveID(HoveredCurveID)
+		, NumBufferedCurves([this]()
+		{
+			int32 NumBufferedCurves = 0;
+			for (const TUniquePtr<IBufferedCurveModel>& BufferedCurve : LocalCurveEditor->GetBufferedCurves())
+			{
+				if (LocalCurveEditor->IsActiveBufferedCurve(BufferedCurve))
+				{
+					++NumBufferedCurves;
+				}
+			}
+			return NumBufferedCurves;
+		}())
+	{}
 
-	// We prioritize key selections over curve selections to reduce the pixel-perfectness needed
-	// to edit the keys (which is more common than curves). Right clicking on a key or an empty space
-	// should show the key menu, otherwise we show the curve menu (ie: right clicking on a curve, not 
-	// directly over a key).
-	if (NumSelectedKeys > 0 && (!HoveredCurveID.IsSet() || ClickedPoint.IsSet()))
+	void Build(FMenuBuilder& MenuBuilder)
 	{
-		MenuBuilder.BeginSection("CurveEditorKeySection", FText::Format(LOCTEXT("CurveEditorKeySection", "{0} Selected {0}|plural(one=Key,other=Keys)"), NumSelectedKeys));
+		// We prioritize key selections over curve selections to reduce the pixel-perfectness needed
+		// to edit the keys (which is more common than curves). Right clicking on a key or an empty space
+		// should show the key menu, otherwise we show the curve menu (ie: right clicking on a curve, not 
+		// directly over a key).
+		const bool bBuildKeyMenu = NumSelectedKeys > 0 && (!HoveredCurveID.IsSet() || ClickedPoint.IsSet());
+		if (bBuildKeyMenu)
 		{
-			bool bIsReadOnly = false;
+			BuildKeyMenu(MenuBuilder);
+		}
+		else
+		{
+			// Test if at least one curve is editable
+			bool bIsReadOnly = true;
+			TSet<FCurveModelID> CurvesToAddTo;
+			for(const FCurveModelID& CurveModelID : CurveEditor->GetEditedCurves())
+			{
+				if (const FCurveModel* CurveModel = CurveEditor->FindCurve(CurveModelID))
+				{
+					if (!CurveModel->IsReadOnly())
+					{
+						bIsReadOnly = false;
+						break;
+					}
+				}
+			}
+
 			if (HoveredCurve)
 			{
-				bIsReadOnly = HoveredCurve->IsReadOnly();
+				BuildHoveredCurve(MenuBuilder, bIsReadOnly);
 			}
 			else
 			{
-				if (ClickedPoint.IsSet())
+				BuildBackgroundMenu(MenuBuilder, bIsReadOnly);
+			}
+		}
+	}
+	
+private:
+
+	/** Content when clicking a key */
+	void BuildKeyMenu(FMenuBuilder& MenuBuilder);
+	/** Adds commands for changing key types */
+	void AppendTangentTypes(FMenuBuilder& MenuBuilder);
+
+	/** Content when clicking a curve */
+	void BuildHoveredCurve(FMenuBuilder& MenuBuilder, bool bIsReadOnly);
+	/** Content when clicking into the background */
+	void BuildBackgroundMenu(FMenuBuilder& MenuBuilder, bool bIsReadOnly);
+
+	/** Appends items that are shared in all menu types (key, curve, and background). */
+	void AppendSharedItems(FMenuBuilder& MenuBuilder);
+};
+	
+void FContextMenuBuilder::BuildKeyMenu(FMenuBuilder& MenuBuilder)
+{
+	MenuBuilder.BeginSection("CurveEditorKeySection", FText::Format(LOCTEXT("CurveEditorKeySection", "{0} Selected {0}|plural(one=Key,other=Keys)"), NumSelectedKeys));
+	{
+		bool bIsReadOnly = false;
+		if (HoveredCurve)
+		{
+			bIsReadOnly = HoveredCurve->IsReadOnly();
+		}
+		else
+		{
+			if (ClickedPoint.IsSet())
+			{
+				if (FCurveModel* ClickedPointCurve = CurveEditor->FindCurve(ClickedPoint->CurveID))
 				{
-					if (FCurveModel* ClickedPointCurve = CurveEditor->FindCurve(ClickedPoint->CurveID))
-					{
-						bIsReadOnly = ClickedPointCurve->IsReadOnly();
-					}
+					bIsReadOnly = ClickedPointCurve->IsReadOnly();
 				}
 			}
+		}
 
-			if (!bIsReadOnly)
-			{
-				// Modify Data
-				MenuBuilder.AddMenuEntry(FGenericCommands::Get().Cut);
-				MenuBuilder.AddMenuEntry(FGenericCommands::Get().Paste);
-				MenuBuilder.AddMenuEntry(FCurveEditorCommands::Get().PasteOverwriteRange);
-				MenuBuilder.AddMenuEntry(FGenericCommands::Get().Copy);
-				MenuBuilder.AddMenuEntry(FGenericCommands::Get().Delete);
+		if (!bIsReadOnly)
+		{
+			// Modify Data
+			MenuBuilder.AddMenuEntry(FGenericCommands::Get().Cut);
+			MenuBuilder.AddMenuEntry(FGenericCommands::Get().Paste);
+			MenuBuilder.AddMenuEntry(FCurveEditorCommands::Get().PasteOverwriteRange);
+			MenuBuilder.AddMenuEntry(FGenericCommands::Get().Copy);
+			MenuBuilder.AddMenuEntry(FGenericCommands::Get().Delete);
+
+			MenuBuilder.AddMenuEntry(FCurveEditorCommands::Get().FlattenTangents);
+			MenuBuilder.AddMenuEntry(FCurveEditorCommands::Get().StraightenTangents);
+			MenuBuilder.AddMenuEntry(FCurveEditorCommands::Get().SmartSnapKeys);
+			
+			MenuBuilder.AddMenuSeparator();
+			AppendTangentTypes(MenuBuilder);
+			MenuBuilder.AddMenuSeparator();
+		}
 
-				MenuBuilder.AddMenuEntry(FCurveEditorCommands::Get().FlattenTangents);
-				MenuBuilder.AddMenuEntry(FCurveEditorCommands::Get().StraightenTangents);
+		// Buffer Curves
+		MenuBuilder.AddMenuEntry(FCurveEditorCommands::Get().BufferVisibleCurves, NAME_None, BufferedCurvesText);
+		if (!bIsReadOnly)
+		{
+			MenuBuilder.AddMenuEntry(FCurveEditorCommands::Get().SwapBufferedCurves, NAME_None, SwapBufferedCurvesText);
+			MenuBuilder.AddMenuEntry(FCurveEditorCommands::Get().ApplyBufferedCurves, NAME_None, ApplyBufferedCurvesText);
+		}
+		MenuBuilder.AddSeparator();
+		
+		AppendSharedItems(MenuBuilder);
+	}
+	MenuBuilder.EndSection();
+}
+
+void FContextMenuBuilder::AppendTangentTypes(FMenuBuilder& MenuBuilder)
+{
+	int32 SupportedTangentTypes = CurveEditor->GetSupportedTangentTypes();
+	if (SupportedTangentTypes & (int32)ECurveEditorTangentTypes::InterpolationCubicSmartAuto)
+	{
+		MenuBuilder.AddMenuEntry(FCurveEditorCommands::Get().InterpolationCubicSmartAuto);
+	};
+	if (SupportedTangentTypes & (int32)ECurveEditorTangentTypes::InterpolationCubicAuto)
+	{
+		MenuBuilder.AddMenuEntry(FCurveEditorCommands::Get().InterpolationCubicAuto);
+	};
+	if (SupportedTangentTypes & (int32)ECurveEditorTangentTypes::InterpolationCubicUser)
+	{
+		MenuBuilder.AddMenuEntry(FCurveEditorCommands::Get().InterpolationCubicUser);
+	}
+	if (SupportedTangentTypes & (int32)ECurveEditorTangentTypes::InterpolationCubicBreak)
+	{
+		MenuBuilder.AddMenuEntry(FCurveEditorCommands::Get().InterpolationCubicBreak);
+	}
+	if (SupportedTangentTypes & (int32)ECurveEditorTangentTypes::InterpolationLinear)
+	{
+		MenuBuilder.AddMenuEntry(FCurveEditorCommands::Get().InterpolationLinear);
+	}
+	if (SupportedTangentTypes & (int32)ECurveEditorTangentTypes::InterpolationConstant)
+	{
+		MenuBuilder.AddMenuEntry(FCurveEditorCommands::Get().InterpolationConstant);
+	}
+	if (SupportedTangentTypes & (int32)ECurveEditorTangentTypes::InterpolationCubicWeighted)
+	{
+		MenuBuilder.AddMenuEntry(FCurveEditorCommands::Get().InterpolationToggleWeighted);
+	}
+}
 
-				MenuBuilder.AddMenuSeparator();
+void FContextMenuBuilder::BuildHoveredCurve(FMenuBuilder& MenuBuilder, bool bIsReadOnly)
+{
+	MenuBuilder.BeginSection("CurveEditorCurveSection", FText::Format(LOCTEXT("CurveNameFormat", "Curve '{0}'"), HoveredCurve->GetLongDisplayName()));
+	{
+		// Modify Curve
+		if (!HoveredCurve->IsReadOnly())
+		{
+			MenuBuilder.AddMenuEntry(FCurveEditorCommands::Get().AddKeyHovered);
+			MenuBuilder.AddMenuEntry(FCurveEditorCommands::Get().PasteKeysHovered);
 
-				// Tangent Types
-				int32 SupportedTangentTypes = CurveEditor->GetSupportedTangentTypes();
-				if (SupportedTangentTypes & (int32)ECurveEditorTangentTypes::InterpolationCubicSmartAuto)
-				{
-					MenuBuilder.AddMenuEntry(FCurveEditorCommands::Get().InterpolationCubicSmartAuto);
-				};
-				if (SupportedTangentTypes & (int32)ECurveEditorTangentTypes::InterpolationCubicAuto)
-				{
-					MenuBuilder.AddMenuEntry(FCurveEditorCommands::Get().InterpolationCubicAuto);
-				};
-				if (SupportedTangentTypes & (int32)ECurveEditorTangentTypes::InterpolationCubicUser)
-				{
-					MenuBuilder.AddMenuEntry(FCurveEditorCommands::Get().InterpolationCubicUser);
-				}
-				if (SupportedTangentTypes & (int32)ECurveEditorTangentTypes::InterpolationCubicBreak)
-				{
-					MenuBuilder.AddMenuEntry(FCurveEditorCommands::Get().InterpolationCubicBreak);
-				}
-				if (SupportedTangentTypes & (int32)ECurveEditorTangentTypes::InterpolationLinear)
-				{
-					MenuBuilder.AddMenuEntry(FCurveEditorCommands::Get().InterpolationLinear);
-				}
-				if (SupportedTangentTypes & (int32)ECurveEditorTangentTypes::InterpolationConstant)
+			MenuBuilder.AddMenuSeparator();
+			
+			MenuBuilder.AddSubMenu(LOCTEXT("PreInfinityText", "Pre-Infinity"), FText(), FNewMenuDelegate::CreateLambda(
+				[](FMenuBuilder& SubMenu)
 				{
-					MenuBuilder.AddMenuEntry(FCurveEditorCommands::Get().InterpolationConstant);
-				}
-				if (SupportedTangentTypes & (int32)ECurveEditorTangentTypes::InterpolationCubicWeighted)
+					SubMenu.AddMenuEntry(FCurveEditorCommands::Get().SetPreInfinityExtrapCycle);
+					SubMenu.AddMenuEntry(FCurveEditorCommands::Get().SetPreInfinityExtrapCycleWithOffset);
+					SubMenu.AddMenuEntry(FCurveEditorCommands::Get().SetPreInfinityExtrapOscillate);
+					SubMenu.AddMenuEntry(FCurveEditorCommands::Get().SetPreInfinityExtrapLinear);
+					SubMenu.AddMenuEntry(FCurveEditorCommands::Get().SetPreInfinityExtrapConstant);
+				})
+			);
+
+			MenuBuilder.AddSubMenu(LOCTEXT("PostInfinityText", "Post-Infinity"), FText(), FNewMenuDelegate::CreateLambda(
+				[](FMenuBuilder& SubMenu)
 				{
-					MenuBuilder.AddMenuEntry(FCurveEditorCommands::Get().InterpolationToggleWeighted);
-				}
+					SubMenu.AddMenuEntry(FCurveEditorCommands::Get().SetPostInfinityExtrapCycle);
+					SubMenu.AddMenuEntry(FCurveEditorCommands::Get().SetPostInfinityExtrapCycleWithOffset);
+					SubMenu.AddMenuEntry(FCurveEditorCommands::Get().SetPostInfinityExtrapOscillate);
+					SubMenu.AddMenuEntry(FCurveEditorCommands::Get().SetPostInfinityExtrapLinear);
+					SubMenu.AddMenuEntry(FCurveEditorCommands::Get().SetPostInfinityExtrapConstant);
+				})
+			);
 
-				MenuBuilder.AddMenuSeparator();
-			}
+			MenuBuilder.AddMenuSeparator();
 
 			// Buffer Curves
 			MenuBuilder.AddMenuEntry(FCurveEditorCommands::Get().BufferVisibleCurves, NAME_None, BufferedCurvesText);
@@ -128,128 +252,57 @@ void FCurveEditorContextMenu::BuildMenu(FMenuBuilder& MenuBuilder, TSharedRef<FC
 				MenuBuilder.AddMenuEntry(FCurveEditorCommands::Get().ApplyBufferedCurves, NAME_None, ApplyBufferedCurvesText);
 			}
 			MenuBuilder.AddSeparator();
-
-			// Select
-			MenuBuilder.AddMenuEntry(FCurveEditorCommands::Get().SelectAllKeys);
-
-			// Filters
-			MenuBuilder.AddMenuEntry(FCurveEditorCommands::Get().OpenUserImplementableFilterWindow);
-			CurveEditor->GetToolbarPromotedFilters()->AppendToBuilder(MenuBuilder);
-			
-			// View
-			MenuBuilder.AddMenuEntry(FCurveEditorCommands::Get().ZoomToFit);
 		}
-		MenuBuilder.EndSection();
+	
+		AppendSharedItems(MenuBuilder);
 	}
-	else
-	{
-		// Test if at least one curve is editable
-		bool bIsReadOnly = true;
-		TSet<FCurveModelID> CurvesToAddTo;
-		for(const FCurveModelID& CurveModelID : CurveEditor->GetEditedCurves())
-		{
-			if (const FCurveModel* CurveModel = CurveEditor->FindCurve(CurveModelID))
-			{
-				if (!CurveModel->IsReadOnly())
-				{
-					bIsReadOnly = false;
-					break;
-				}
-			}
-		}
+	MenuBuilder.EndSection();
+}
 
-		if (HoveredCurve)
+void FContextMenuBuilder::BuildBackgroundMenu(FMenuBuilder& MenuBuilder, bool bIsReadOnly)
+{
+	MenuBuilder.BeginSection("CurveEditorAllCurveSections", LOCTEXT("CurveEditorAllCurveSections", "All Curves"));
+	{
+		if (!bIsReadOnly)
 		{
-			MenuBuilder.BeginSection("CurveEditorCurveSection", FText::Format(LOCTEXT("CurveNameFormat", "Curve '{0}'"), HoveredCurve->GetLongDisplayName()));
-			{
-				// Modify Curve
-				if (!HoveredCurve->IsReadOnly())
-				{
-					MenuBuilder.AddMenuEntry(FCurveEditorCommands::Get().AddKeyHovered);
-					MenuBuilder.AddMenuEntry(FCurveEditorCommands::Get().PasteKeysHovered);
-
-					MenuBuilder.AddMenuSeparator();
-					
-					MenuBuilder.AddSubMenu(LOCTEXT("PreInfinityText", "Pre-Infinity"), FText(), FNewMenuDelegate::CreateLambda(
-						[](FMenuBuilder& SubMenu)
-						{
-							SubMenu.AddMenuEntry(FCurveEditorCommands::Get().SetPreInfinityExtrapCycle);
-							SubMenu.AddMenuEntry(FCurveEditorCommands::Get().SetPreInfinityExtrapCycleWithOffset);
-							SubMenu.AddMenuEntry(FCurveEditorCommands::Get().SetPreInfinityExtrapOscillate);
-							SubMenu.AddMenuEntry(FCurveEditorCommands::Get().SetPreInfinityExtrapLinear);
-							SubMenu.AddMenuEntry(FCurveEditorCommands::Get().SetPreInfinityExtrapConstant);
-						})
-					);
-
-					MenuBuilder.AddSubMenu(LOCTEXT("PostInfinityText", "Post-Infinity"), FText(), FNewMenuDelegate::CreateLambda(
-						[](FMenuBuilder& SubMenu)
-						{
-							SubMenu.AddMenuEntry(FCurveEditorCommands::Get().SetPostInfinityExtrapCycle);
-							SubMenu.AddMenuEntry(FCurveEditorCommands::Get().SetPostInfinityExtrapCycleWithOffset);
-							SubMenu.AddMenuEntry(FCurveEditorCommands::Get().SetPostInfinityExtrapOscillate);
-							SubMenu.AddMenuEntry(FCurveEditorCommands::Get().SetPostInfinityExtrapLinear);
-							SubMenu.AddMenuEntry(FCurveEditorCommands::Get().SetPostInfinityExtrapConstant);
-						})
-					);
-
-					MenuBuilder.AddMenuSeparator();
-
-					// Buffer Curves
-					MenuBuilder.AddMenuEntry(FCurveEditorCommands::Get().BufferVisibleCurves, NAME_None, BufferedCurvesText);
-					if (!bIsReadOnly)
-					{
-						MenuBuilder.AddMenuEntry(FCurveEditorCommands::Get().SwapBufferedCurves, NAME_None, SwapBufferedCurvesText);
-						MenuBuilder.AddMenuEntry(FCurveEditorCommands::Get().ApplyBufferedCurves, NAME_None, ApplyBufferedCurvesText);
-					}
-					MenuBuilder.AddSeparator();
-				}
+			// Modify Curves
+			MenuBuilder.AddMenuEntry(FCurveEditorCommands::Get().AddKeyToAllCurves);
 
-				// Select
-				MenuBuilder.AddMenuEntry(FCurveEditorCommands::Get().SelectAllKeys);
-
-				// Filters
-				MenuBuilder.AddMenuEntry(FCurveEditorCommands::Get().OpenUserImplementableFilterWindow);
-				CurveEditor->GetToolbarPromotedFilters()->AppendToBuilder(MenuBuilder);
+			MenuBuilder.AddSeparator();
 
-				// View
-				MenuBuilder.AddMenuEntry(FCurveEditorCommands::Get().ZoomToFit);
+			// Buffer Curves
+			MenuBuilder.AddMenuEntry(FCurveEditorCommands::Get().BufferVisibleCurves, NAME_None, BufferedCurvesText);
+			if (!bIsReadOnly)
+			{
+				MenuBuilder.AddMenuEntry(FCurveEditorCommands::Get().SwapBufferedCurves, NAME_None, SwapBufferedCurvesText);
+				MenuBuilder.AddMenuEntry(FCurveEditorCommands::Get().ApplyBufferedCurves, NAME_None, ApplyBufferedCurvesText);
 			}
-			MenuBuilder.EndSection();
+			MenuBuilder.AddSeparator();
 		}
-		else
-		{
-			MenuBuilder.BeginSection("CurveEditorAllCurveSections", LOCTEXT("CurveEditorAllCurveSections", "All Curves"));
-			{
-				if (!bIsReadOnly)
-				{
-					// Modify Curves
-					MenuBuilder.AddMenuEntry(FCurveEditorCommands::Get().AddKeyToAllCurves);
 
-					MenuBuilder.AddSeparator();
+		AppendSharedItems(MenuBuilder);
+	}
+	MenuBuilder.EndSection();
+}
 
-					// Buffer Curves
-					MenuBuilder.AddMenuEntry(FCurveEditorCommands::Get().BufferVisibleCurves, NAME_None, BufferedCurvesText);
-					if (!bIsReadOnly)
-					{
-						MenuBuilder.AddMenuEntry(FCurveEditorCommands::Get().SwapBufferedCurves, NAME_None, SwapBufferedCurvesText);
-						MenuBuilder.AddMenuEntry(FCurveEditorCommands::Get().ApplyBufferedCurves, NAME_None, ApplyBufferedCurvesText);
-					}
-					MenuBuilder.AddSeparator();
-				}
+void FContextMenuBuilder::AppendSharedItems(FMenuBuilder& MenuBuilder)
+{
+	// Select
+	MenuBuilder.AddMenuEntry(FCurveEditorCommands::Get().SelectAllKeys);
 
-				// Select
-				MenuBuilder.AddMenuEntry(FCurveEditorCommands::Get().SelectAllKeys);
+	// Filters
+	MenuBuilder.AddMenuEntry(FCurveEditorCommands::Get().OpenUserImplementableFilterWindow);
+	CurveEditor->GetToolbarPromotedFilters()->AppendToBuilder(MenuBuilder);
+			
+	// View
+	MenuBuilder.AddMenuEntry(FCurveEditorCommands::Get().ZoomToFit);
+}
+}
 
-				// Filters
-				MenuBuilder.AddMenuEntry(FCurveEditorCommands::Get().OpenUserImplementableFilterWindow);
-				CurveEditor->GetToolbarPromotedFilters()->AppendToBuilder(MenuBuilder);
-				
-				// View
-				MenuBuilder.AddMenuEntry(FCurveEditorCommands::Get().ZoomToFit);
-			}
-			MenuBuilder.EndSection();
-		}
-	}
+void FCurveEditorContextMenu::BuildMenu(FMenuBuilder& MenuBuilder, TSharedRef<FCurveEditor> CurveEditor, TOptional<FCurvePointHandle> ClickedPoint, TOptional<FCurveModelID> HoveredCurveID)
+{
+	UE::CurveEditor::ContextDetail::FContextMenuBuilder Builder(CurveEditor, ClickedPoint, HoveredCurveID);
+	Builder.Build(MenuBuilder);
 }
 
 #undef LOCTEXT_NAMESPACE
diff --git a/Engine/Source/Editor/CurveEditor/Private/Filters/CurveEditorSmartSnapFilter.cpp b/Engine/Source/Editor/CurveEditor/Private/Filters/CurveEditorSmartSnapFilter.cpp
new file mode 100644
--- /dev/null
+++ b/Engine/Source/Editor/CurveEditor/Private/Filters/CurveEditorSmartSnapFilter.cpp
@@ -0,0 +1,23 @@
+﻿// Copyright Epic Games, Inc. All Rights Reserved.
+
+#include "Filters/CurveEditorSmartSnapFilter.h"
+
+#include "Misc/SmartSnap.h"
+
+void UCurveEditorSmartSnapFilter::ApplyFilter_Impl(
+	TSharedRef<FCurveEditor> InCurveEditor,
+	const TMap<FCurveModelID, FKeyHandleSet>& InKeysToOperateOn, TMap<FCurveModelID, FKeyHandleSet>& OutKeysToSelect
+	)
+{
+	UE::CurveEditor::EnumerateSmartSnappableKeys(*InCurveEditor, InKeysToOperateOn, OutKeysToSelect,
+		[&OutKeysToSelect](const FCurveModelID& CurveModelId, FCurveModel& CurveModel, const UE::CurveEditor::FSmartSnapResult& SnapResult)
+		{
+			CurveModel.Modify();
+			ApplySmartSnap(CurveModel, SnapResult);
+		});
+}
+
+bool UCurveEditorSmartSnapFilter::CanApplyFilter_Impl(TSharedRef<FCurveEditor> InCurveEditor)
+{
+	return UE::CurveEditor::CanSmartSnapSelection(InCurveEditor->GetSelection());
+}
diff --git a/Engine/Source/Editor/CurveEditor/Private/Misc/Mirror/MirrorUtils.cpp b/Engine/Source/Editor/CurveEditor/Private/Misc/Mirror/MirrorUtils.cpp
--- a/Engine/Source/Editor/CurveEditor/Private/Misc/Mirror/MirrorUtils.cpp
+++ b/Engine/Source/Editor/CurveEditor/Private/Misc/Mirror/MirrorUtils.cpp
@@ -9,19 +9,6 @@
 
 namespace UE::CurveEditor
 {
-void MirrorTangentsOnMidpoint(const TMap<FCurveModelID, TArray<FKeyHandle>>& InKeysToMirror, FCurveEditor& InCurveEditor)
-{
-	for (const TPair<FCurveModelID, TArray<FKeyHandle>>& Pair : InKeysToMirror)
-	{
-		if (const FCurveModel* Model = InCurveEditor.FindCurve(Pair.Key))
-		{
-			const auto[Min, Max] = FindMinMaxHeight(Pair.Value, *Model);
-			const double Midpoint = Min + (Max - Min) * 0.5;
-			MirrorTangents(Pair.Value, Pair.Key, Min, Max, Midpoint, InCurveEditor);
-		}
-	}
-}
-
 FCurveBounds FindMinMaxHeight(TConstArrayView<FKeyHandle> InKeys, const FCurveModel& InCurveModel)
 {
 	if (!ensure(!InKeys.IsEmpty()))
diff --git a/Engine/Source/Editor/CurveEditor/Private/Misc/SmartSnap.cpp b/Engine/Source/Editor/CurveEditor/Private/Misc/SmartSnap.cpp
new file mode 100644
--- /dev/null
+++ b/Engine/Source/Editor/CurveEditor/Private/Misc/SmartSnap.cpp
@@ -0,0 +1,185 @@
+﻿// Copyright Epic Games, Inc. All Rights Reserved.
+
+#include "SmartSnap.h"
+
+#include "Algo/AllOf.h"
+#include "Algo/AnyOf.h"
+#include "CurveEditor.h"
+
+#include <cmath>
+
+namespace UE::CurveEditor
+{
+bool CanSmartSnapSelection(const FCurveEditorSelection& InSelection)
+{
+	return Algo::AnyOf(InSelection.GetAll(), [](const TPair<FCurveModelID, FKeyHandleSet>& Pair)
+	{
+		const FKeyHandleSet& KeySelection = Pair.Value;
+		return Algo::AnyOf(KeySelection.AsArray(), [&KeySelection](const FKeyHandle& InHandle)
+		{
+			return KeySelection.PointType(InHandle) == ECurvePointType::Key;
+		});
+	});
+}
+	
+namespace SmartSnapDetail
+{
+static TArray<FKeyHandle> ReturnOnlyKeys(const FKeyHandleSet& InSelection)
+{
+	TArray<FKeyHandle> Result;
+	Result.Reserve(InSelection.Num());
+	for (const FKeyHandle& KeyHandle : InSelection.AsArray())
+	{
+		if (InSelection.PointType(KeyHandle) == ECurvePointType::Key)
+		{
+			Result.Add(KeyHandle);
+		}
+	}
+	return Result;
+}
+
+static FFrameRate GetCurveEditorFrameRate(const FCurveEditor& InCurveEditor)
+{
+	const TSharedPtr<ITimeSliderController> Controller = InCurveEditor.GetTimeSliderController();
+	return Controller ? Controller->GetDisplayRate() : FFrameRate{};
+}
+}
+	
+void EnumerateSmartSnappableKeys(
+	const FCurveEditor& InCurveEditor,
+	const TMap<FCurveModelID, FKeyHandleSet>& InKeysToOperateOn,
+	TMap<FCurveModelID, FKeyHandleSet>& OutKeysToSelect,
+	const TFunctionRef<void(const FCurveModelID& CurveModelId, FCurveModel& CurveModel, const FSmartSnapResult& SnapResult)>& InProcessSmartSnapping
+	)
+{
+	const FFrameRate FrameRate = SmartSnapDetail::GetCurveEditorFrameRate(InCurveEditor);
+	for (const TPair<FCurveModelID, FKeyHandleSet>& Pair : InKeysToOperateOn)
+	{
+		FCurveModel* CurveModel = InCurveEditor.FindCurve(Pair.Key);
+		if (!CurveModel)
+		{
+			continue;
+		}
+
+		// Exclude tangent handles
+		const FKeyHandleSet& KeySelection = Pair.Value;
+		const bool bSelectionContainsOnlyKeys = Algo::AllOf(KeySelection.AsArray(), [&KeySelection](const FKeyHandle& InHandle)
+		{
+			return KeySelection.PointType(InHandle) == ECurvePointType::Key;
+		});
+		// Avoid TArray allocation when user has only selected key handles
+		const TArray<FKeyHandle> FilteredKeys = bSelectionContainsOnlyKeys ? TArray<FKeyHandle>{} : SmartSnapDetail::ReturnOnlyKeys(KeySelection);
+		const TConstArrayView<FKeyHandle> Keys = bSelectionContainsOnlyKeys ? KeySelection.AsArray() : FilteredKeys;
+		const int32 NumKeys = Keys.Num();
+		if (NumKeys == 0)
+		{
+			continue;
+		}
+		
+		TArray<FKeyPosition> Positions;
+		Positions.SetNumUninitialized(NumKeys);
+		CurveModel->GetKeyPositions(Keys, Positions);
+
+		const FSmartSnapResult SnappingResult = ComputeSmartSnap(*CurveModel, Keys, Positions, FrameRate);
+		for (const FKeyHandle& InHandle : SnappingResult.UpdatedKeys)
+		{
+			OutKeysToSelect.FindOrAdd(Pair.Key).Add(InHandle, ECurvePointType::Key);
+		}
+		InProcessSmartSnapping(Pair.Key, *CurveModel, SnappingResult);
+	}
+}
+
+namespace SmartSnapDetail
+{
+struct FFrameData
+{
+	FKeyHandle ClosestHandle = FKeyHandle::Invalid();
+	FFrameTime AbsDistToFrame;
+};
+
+static TMap<FFrameNumber, FFrameData> ComputeClosestFrames(
+	TConstArrayView<FKeyHandle> InHandles, TConstArrayView<FKeyPosition> InPositions, const FFrameRate& InFrameRate
+	)
+{
+	TMap<FFrameNumber, FFrameData> FrameToData;
+	for (int32 Index = 0; Index < InHandles.Num(); ++Index)
+	{
+		const FKeyHandle& KeyHandle = InHandles[Index];
+		const FFrameTime SubFrame = InFrameRate.AsFrameTime(InPositions[Index].InputValue);
+		const FFrameTime FrameNumber = SubFrame.RoundToFrame();
+		
+		const FFrameTime AbsDistToFrame = FMath::Max(SubFrame, FrameNumber) - FMath::Min(SubFrame, FrameNumber);
+		FFrameData& FrameData = FrameToData.FindOrAdd(FrameNumber.FrameNumber, FFrameData{ KeyHandle, AbsDistToFrame });
+		const bool bIsFirstKeyOnFrame = FrameData.ClosestHandle == KeyHandle; 
+		if (bIsFirstKeyOnFrame)
+		{
+			continue;
+		}
+		
+		const bool bIsCloserToFrame = AbsDistToFrame < FrameData.AbsDistToFrame;
+		if (bIsCloserToFrame)
+		{
+			FrameData.ClosestHandle = KeyHandle;
+			FrameData.AbsDistToFrame = AbsDistToFrame;
+		}
+	}
+	return FrameToData;
+}
+static FSmartSnapResult MoveKeysOntoFrames(
+	const FCurveModel& InModel, const TMap<FFrameNumber, FFrameData>& FrameToData, const FFrameRate& InFrameRate
+	)
+{
+	FSmartSnapResult Result;
+	for (const TPair<FFrameNumber, FFrameData>& Frame : FrameToData)
+	{
+		const FKeyHandle& KeyHandle = Frame.Value.ClosestHandle;
+		const FFrameTime FrameNumber = Frame.Key;
+		
+		FKeyPosition Position;
+		Position.InputValue = InFrameRate.AsSeconds(FrameNumber);
+		InModel.Evaluate(Position.InputValue, Position.OutputValue);
+		Result.NewPositions.Add(Position);
+		Result.UpdatedKeys.Add(KeyHandle);
+	}
+	return Result;
+}
+}
+
+FSmartSnapResult ComputeSmartSnap(
+	const FCurveModel& InModel, TConstArrayView<FKeyHandle> InHandles, TConstArrayView<FKeyPosition> InPositions, const FFrameRate& InFrameRate
+	)
+{
+	check(InHandles.Num() == InPositions.Num());
+
+	// 1. We'll compute all the frames covered by the keys, and the single key that is closest to it.
+	// Example: If key 1 is at 2.6 and key 2 at 2.7, we'd move key 2 to frame 3.0.
+	// This retains the shape of the curve a bit better (as opposed to taking a 'random' one without criteria).
+	const TMap<FFrameNumber, SmartSnapDetail::FFrameData> FrameToData = SmartSnapDetail::ComputeClosestFrames(InHandles, InPositions, InFrameRate);
+	// 2. The key closest to its assigned frame is moved there by evaluating the curve.
+	FSmartSnapResult Result = MoveKeysOntoFrames(InModel, FrameToData, InFrameRate);
+
+	// 3.1 All keys were moved? Done.
+	const int32 NumToRemove = InHandles.Num() - Result.NewPositions.Num();
+	if (!NumToRemove)
+	{
+		return Result;
+	}
+	// 3.2 All keys that were not moved are removed.
+	Result.RemovedKeys.Reserve(NumToRemove);
+	for (const FKeyHandle& Handle : InHandles)
+	{
+		if (!Result.UpdatedKeys.Contains(Handle))
+		{
+			Result.RemovedKeys.Add(Handle);
+		}
+	}
+	
+	return Result;
+}
+
+void ApplySmartSnap(FCurveModel& InModel, const FSmartSnapResult& InSmartSnap, double InCurrentTime)
+{
+	InModel.RemoveKeys(InSmartSnap.RemovedKeys, InCurrentTime);
+	InModel.SetKeyPositions(InSmartSnap.UpdatedKeys, InSmartSnap.NewPositions);
+}
+}
diff --git a/Engine/Source/Editor/CurveEditor/Private/Misc/SmartSnap.h b/Engine/Source/Editor/CurveEditor/Private/Misc/SmartSnap.h
new file mode 100644
--- /dev/null
+++ b/Engine/Source/Editor/CurveEditor/Private/Misc/SmartSnap.h
@@ -0,0 +1,46 @@
+﻿// Copyright Epic Games, Inc. All Rights Reserved.
+
+#pragma once
+
+#include "CurveDataAbstraction.h"
+#include "Curves/KeyHandle.h"
+#include "Containers/Array.h"
+#include "Templates/FunctionFwd.h"
+
+struct FKeyHandleSet;
+class FCurveEditor;
+class FCurveModel;
+class ITimeSliderController;
+struct FCurveEditorSelection;
+struct FCurveModelID;
+
+namespace UE::CurveEditor
+{
+/** @return Whether the selection contains keys that can be snapped. */
+bool CanSmartSnapSelection(const FCurveEditorSelection& InSelection);
+
+struct FSmartSnapResult
+{
+	TArray<FKeyHandle> RemovedKeys;
+	TArray<FKeyHandle> UpdatedKeys;
+	TArray<FKeyPosition> NewPositions;
+};
+
+/** Calls InProcessSmartSnapping for each curve on which smart snapping can be performed. */
+void EnumerateSmartSnappableKeys(
+	const FCurveEditor& InCurveEditor,
+	const TMap<FCurveModelID, FKeyHandleSet>& InKeysToOperateOn,
+	TMap<FCurveModelID, FKeyHandleSet>& OutKeysToSelect,
+	const TFunctionRef<void(const FCurveModelID& CurveModelId, FCurveModel& CurveModel, const FSmartSnapResult& SnapResult)>& InProcessSmartSnapping
+	);
+	
+/** Computes how to modify the curve model: Tries to snap keys to the closest whole frame placing the key on the intersection of curve and an imaginary vertical frame.*/
+FSmartSnapResult ComputeSmartSnap(
+	const FCurveModel& InModel, TConstArrayView<FKeyHandle> InHandles, TConstArrayView<FKeyPosition> InPositions, const FFrameRate& InFrameRate
+	);
+
+/** Applies the computed smart snapping to InModel. */
+void ApplySmartSnap(
+	FCurveModel& InModel, const FSmartSnapResult& InSmartSnap, double CurrentTime = 0.0
+	);
+}
diff --git a/Engine/Source/Editor/CurveEditor/Public/CurveEditor.h b/Engine/Source/Editor/CurveEditor/Public/CurveEditor.h
--- a/Engine/Source/Editor/CurveEditor/Public/CurveEditor.h
+++ b/Engine/Source/Editor/CurveEditor/Public/CurveEditor.h
@@ -561,6 +561,11 @@ public:
 	*/
 	bool CanFlattenOrStraightenSelection() const;
 
+	/** Snaps the selected keys to their nearest whole frames and positions them on the frame intersection with the frame. */
+	void SmartSnapSelection();
+	/** @return Whether there is any key selection. */
+	bool CanSmartSnapSelection() const;
+
 public:
 
 	/**
diff --git a/Engine/Source/Editor/CurveEditor/Public/CurveEditorCommands.h b/Engine/Source/Editor/CurveEditor/Public/CurveEditorCommands.h
--- a/Engine/Source/Editor/CurveEditor/Public/CurveEditorCommands.h
+++ b/Engine/Source/Editor/CurveEditor/Public/CurveEditorCommands.h
@@ -51,6 +51,8 @@ public:
 	TSharedPtr<FUICommandInfo> FlattenTangents;
 	TSharedPtr<FUICommandInfo> StraightenTangents;
 
+	TSharedPtr<FUICommandInfo> SmartSnapKeys;
+
 	TSharedPtr< FUICommandInfo > MatchLastTangentToFirst;
 	TSharedPtr< FUICommandInfo > MatchFirstTangentToLast;
 	
diff --git a/Engine/Source/Editor/CurveEditor/Public/Filters/CurveEditorSmartSnapFilter.h b/Engine/Source/Editor/CurveEditor/Public/Filters/CurveEditorSmartSnapFilter.h
new file mode 100644
--- /dev/null
+++ b/Engine/Source/Editor/CurveEditor/Public/Filters/CurveEditorSmartSnapFilter.h
@@ -0,0 +1,25 @@
+﻿// Copyright Epic Games, Inc. All Rights Reserved.
+
+#pragma once
+
+#include "CoreMinimal.h"
+#include "CurveEditorFilterBase.h"
+#include "CurveEditorSmartSnapFilter.generated.h"
+
+/**
+ * 
+ */
+UCLASS(DisplayName = "Smart Snap", meta = (
+	CurveEditorLabel = "Smart Snap",
+	CurveEditorDescription = "Snaps selected keys to the closest whole frame.\nThe key is placed where the curve intersects with the frame (imagine vertical line going down from the frame).")
+	)
+class CURVEEDITOR_API UCurveEditorSmartSnapFilter : public UCurveEditorFilterBase
+{
+	GENERATED_BODY()
+protected:
+
+	//~ Begin UCurveEditorFilterBase Interface
+	virtual void ApplyFilter_Impl(TSharedRef<FCurveEditor> InCurveEditor, const TMap<FCurveModelID, FKeyHandleSet>& InKeysToOperateOn, TMap<FCurveModelID, FKeyHandleSet>& OutKeysToSelect) override;
+	virtual bool CanApplyFilter_Impl(TSharedRef<FCurveEditor> InCurveEditor) override;
+	//~ End UCurveEditorFilterBase Interface
+};
