Fix variable shadowing error on Mac

diff --git a/Engine/Source/Editor/CurveEditor/Private/CurveEditorContextMenu.cpp b/Engine/Source/Editor/CurveEditor/Private/CurveEditorContextMenu.cpp
--- a/Engine/Source/Editor/CurveEditor/Private/CurveEditorContextMenu.cpp
+++ b/Engine/Source/Editor/CurveEditor/Private/CurveEditorContextMenu.cpp
@@ -38,7 +38,7 @@ public:
 	const int32 NumSelectedKeys = CurveEditor->GetSelection().Count();
 	const TSharedPtr<FCurveEditor> LocalCurveEditor = CurveEditor->AsShared();
 	const int32 NumActiveCurves = LocalCurveEditor->GetSelectionFromTreeAndKeys().Num();
-	const int32 NumBufferedCurves = 0;
+	int32 NumBufferedCurves = 0;
 	
 	// We change the name to reflect the current number of curves selected.
 	const TAttribute<FText> BufferedCurvesText
@@ -54,19 +54,15 @@ public:
 		: CurveEditor(CurveEditor)
 		, ClickedPoint(ClickedPoint)
 		, HoveredCurveID(HoveredCurveID)
-		, NumBufferedCurves([this]()
+	{
+		for (const TUniquePtr<IBufferedCurveModel>& BufferedCurve : LocalCurveEditor->GetBufferedCurves())
 		{
-			int32 NumBufferedCurves = 0;
-			for (const TUniquePtr<IBufferedCurveModel>& BufferedCurve : LocalCurveEditor->GetBufferedCurves())
+			if (LocalCurveEditor->IsActiveBufferedCurve(BufferedCurve))
 			{
-				if (LocalCurveEditor->IsActiveBufferedCurve(BufferedCurve))
-				{
-					++NumBufferedCurves;
-				}
+				++NumBufferedCurves;
 			}
-			return NumBufferedCurves;
-		}())
-	{}
+		}
+	}
 
 	void Build(FMenuBuilder& MenuBuilder)
 	{
