CurveEditor: Key is not added to curves using the rich curve model when curve editor zooms in

diff --git a/Engine/Source/Editor/CurveEditor/Private/RichCurveEditorModel.cpp b/Engine/Source/Editor/CurveEditor/Private/RichCurveEditorModel.cpp
--- a/Engine/Source/Editor/CurveEditor/Private/RichCurveEditorModel.cpp
+++ b/Engine/Source/Editor/CurveEditor/Private/RichCurveEditorModel.cpp
@@ -244,7 +244,10 @@ void FRichCurveEditorModel::AddKeys(TArrayView<const FKeyPosition> InKeyPosition
 				FKeyAttributes Attributes = InKeyAttributes[Index];
 
 				const TRange<double> InputRange = ClampInputRange.Get();
-				FKeyHandle     NewHandle = RichCurve.UpdateOrAddKey(FMath::Clamp(Position.InputValue, InputRange.GetLowerBoundValue(), InputRange.GetUpperBoundValue()), Position.OutputValue);
+
+				constexpr bool bUnwindRotation = false;
+				constexpr float Tolearance = 0.0f; 
+				const FKeyHandle NewHandle = RichCurve.UpdateOrAddKey(FMath::Clamp(Position.InputValue, InputRange.GetLowerBoundValue(), InputRange.GetUpperBoundValue()), Position.OutputValue, bUnwindRotation, Tolearance);
 				if (NewHandle != FKeyHandle::Invalid())
 				{
 					NewKeyHandles[Index] = NewHandle;
