Curve Editor: Fix PRAGMA_DISABLE_DEPRECATION_WARNINGS macro appears a second time without matching PRAGMA_ENABLE_DEPRECATION_WARNINGS macro

#rnx

diff --git a/Engine/Source/Editor/CurveEditor/Private/SCurveEditorView.cpp b/Engine/Source/Editor/CurveEditor/Private/SCurveEditorView.cpp
--- a/Engine/Source/Editor/CurveEditor/Private/SCurveEditorView.cpp
+++ b/Engine/Source/Editor/CurveEditor/Private/SCurveEditorView.cpp
@@ -48,7 +48,7 @@ SCurveEditorView::SCurveEditorView()
 	CachedValues.CachedSelectionSerialNumber = 0xFFFFFFFF;
 	CachedValues.CachedGeometrySize.X = -1;
 	CachedValues.CachedGeometrySize.Y = -1;
-	PRAGMA_DISABLE_DEPRECATION_WARNINGS
+	PRAGMA_ENABLE_DEPRECATION_WARNINGS
 }
 
 PRAGMA_DISABLE_DEPRECATION_WARNINGS // Allow deprecated CachedCurveParams to destruct
@@ -98,7 +98,11 @@ void SCurveEditorView::UpdateCurveViewTransformsFromModels()
 
 	if (bTransformChanged)
 	{
+		CurveEditor->GetCurveDrawParamsCache().Invalidate(SharedThis(this));
+		
+		PRAGMA_DISABLE_DEPRECATION_WARNINGS
 		CurveCacheFlags = ECurveCacheFlags::All;
+		PRAGMA_ENABLE_DEPRECATION_WARNINGS
 	}
 }
 
@@ -591,7 +595,7 @@ void SCurveEditorView::GetCurveDrawParams(TArray<FCurveDrawParams>& OutDrawParam
 
 		PRAGMA_DISABLE_DEPRECATION_WARNINGS
 		GetCurveDrawParam(CurveEditor, Pair.Key, CurveModel, Params);
-		PRAGMA_DISABLE_DEPRECATION_WARNINGS
+		PRAGMA_ENABLE_DEPRECATION_WARNINGS
 		
 		OutDrawParams.Add(MoveTemp(Params));
 	}
@@ -770,6 +774,9 @@ void SCurveEditorView::CheckCacheAndInvalidateIfNeeded()
 	}
 	else
 	{
+		// Using deprecated methods to allow testing if the deprecated API remains functional for the time being.
+		// Can be updated to DrawParamsCache::GetDrawParamSynchronous when GetCurveDrawParams is removed.
+		PRAGMA_DISABLE_DEPRECATION_WARNINGS
 		CachedValues.CachedOutputMin = OutputMin;
 		CachedValues.CachedOutputMax = OutputMax;
 		GetInputBounds(CachedValues.CachedInputMin, CachedValues.CachedInputMax);
@@ -777,9 +784,6 @@ void SCurveEditorView::CheckCacheAndInvalidateIfNeeded()
 	
 		CurveEditor->GetCurveDrawParamsCache().CachedDrawParams.Reset();
 
-		PRAGMA_DISABLE_DEPRECATION_WARNINGS
-		// Using deprecated methods to allow testing if the deprecated API remains functional for the time being.
-		// Can be updated to DrawParamsCache::GetDrawParamSynchronous when GetCurveDrawParams is removed.
 		GetCurveDrawParams(CurveEditor->GetCurveDrawParamsCache().CachedDrawParams);
 		PRAGMA_ENABLE_DEPRECATION_WARNINGS
 
