Curve Editor: Fix curve cache optimization are now disabled by default as expected

#rnx

diff --git a/Engine/Source/Editor/CurveEditor/Private/CurveEditorCurveDrawParamsCache.cpp b/Engine/Source/Editor/CurveEditor/Private/CurveEditorCurveDrawParamsCache.cpp
--- a/Engine/Source/Editor/CurveEditor/Private/CurveEditorCurveDrawParamsCache.cpp
+++ b/Engine/Source/Editor/CurveEditor/Private/CurveEditorCurveDrawParamsCache.cpp
@@ -12,7 +12,7 @@
 
 namespace UE::CurveEditor
 {
-	int32 GCurveEditorUseCurveCachePool = 1;
+	int32 GCurveEditorUseCurveCachePool = 0;
 
 	static FAutoConsoleVariableRef CCurveEditorUseCurveCachePool(
 		TEXT("CurveEditor.UseCurveCachePool"),
