Sequencer: Remove duplicate keys on mouse release in the Curve Editor and when deselecting keys in Sequencer

For the Curve Editor:
TMovieSceneChannelData::MoveKey now has an argument to determine whether duplicate keys should be removed when setting new key times.
bRemoveDuplicateKeys defaults to false and is true when the Curve Editor calls MoveKey with a change type of ValueSet. This allows duplicates to continue to occur when dragging interactively.
Removed the special code in MoveKeys and TransformTool that does this separately.

For Sequencer:
Override FKeySelection::Deselect and FKeySelection::Empty to look for duplicate keys when keys are deselected.

diff --git a/Engine/Plugins/Animation/ControlRig/Source/ControlRigEditor/Private/ControlRigSpaceChannelCurveModel.cpp b/Engine/Plugins/Animation/ControlRig/Source/ControlRigEditor/Private/ControlRigSpaceChannelCurveModel.cpp
--- a/Engine/Plugins/Animation/ControlRig/Source/ControlRigEditor/Private/ControlRigSpaceChannelCurveModel.cpp
+++ b/Engine/Plugins/Animation/ControlRig/Source/ControlRigEditor/Private/ControlRigSpaceChannelCurveModel.cpp
@@ -240,7 +240,8 @@ void FControlRigSpaceChannelCurveModel::SetKeyPositions(TArrayView<const FKeyHan
 			{
 				FFrameNumber NewTime = (InKeyPositions[Index].InputValue * TickResolution).FloorToFrame();
 
-				KeyIndex = ChannelData.MoveKey(KeyIndex, NewTime);
+				const bool bRemoveDuplicates = ChangeType == EPropertyChangeType::ValueSet;
+				KeyIndex = ChannelData.MoveKey(KeyIndex, NewTime, bRemoveDuplicates);
 				Section->ExpandToFrame(NewTime);
 			}
 		}
diff --git a/Engine/Plugins/Editor/CurveEditorTools/Source/CurveEditorTools/Private/Tools/CurveEditorTransformTool.cpp b/Engine/Plugins/Editor/CurveEditorTools/Source/CurveEditorTools/Private/Tools/CurveEditorTransformTool.cpp
--- a/Engine/Plugins/Editor/CurveEditorTools/Source/CurveEditorTools/Private/Tools/CurveEditorTransformTool.cpp
+++ b/Engine/Plugins/Editor/CurveEditorTools/Source/CurveEditorTools/Private/Tools/CurveEditorTransformTool.cpp
@@ -1099,55 +1099,6 @@ void FCurveEditorTransformTool::OnDragEnd()
 	{
 		if (FCurveModel* Curve = CurveEditor->FindCurve(KeyData.CurveID))
 		{
-			// Gather all the final key time positions
-			TArray<FKeyPosition> KeyTimes;
-			for (int32 KeyIndex = 0; KeyIndex < KeyData.Handles.Num(); ++KeyIndex)
-			{
-				const FKeyHandle& KeyHandle = KeyData.Handles[KeyIndex];
-				FKeyPosition KeyTime = KeyData.LastDraggedKeyPositions[KeyIndex];
-				KeyTimes.Add(KeyTime);
-			}
-
-			// For each key time, look for all the keys that match
-			TArray<FKeyHandle> KeysToRemove;
-			for (const FKeyPosition& KeyTime : KeyTimes)
-			{
-				TArray<FKeyHandle> KeysInRange;
-				Curve->GetKeys(KeyTime.InputValue, KeyTime.InputValue, TNumericLimits<double>::Lowest(), TNumericLimits<double>::Max(), KeysInRange);
-
-				// If there's more than 1 key at this time, remove all but the keys that moved the largest amount
-				if (KeysInRange.Num() > 1)
-				{
-					double MaxDist = TNumericLimits<double>::Min();
-					FKeyHandle MaxKeyHandle;
-					for (const FKeyHandle& KeyInRange : KeysInRange)
-					{
-						int32 KeyDataIndex = KeyData.Handles.Find(KeyInRange);
-						if (KeyDataIndex != INDEX_NONE)
-						{
-							double Dist = FMath::Abs(KeyData.StartKeyPositions[KeyDataIndex].InputValue - KeyData.LastDraggedKeyPositions[KeyDataIndex].InputValue);
-							if (Dist > MaxDist)
-							{
-								MaxKeyHandle = KeyInRange;
-								MaxDist = Dist;
-							}
-						}
-					}
-
-					for (const FKeyHandle& KeyInRange : KeysInRange)
-					{
-						if (KeyInRange != MaxKeyHandle)
-						{
-							KeysToRemove.Add(KeyInRange);
-						}
-					}
-				}
-			}
-
-			// Remove any keys that overlap before moving new keys on top
-			Curve->RemoveKeys(KeysToRemove, 0.0);
-
-			// Then, move the keys
 			Curve->SetKeyPositions(KeyData.Handles, KeyData.LastDraggedKeyPositions, EPropertyChangeType::ValueSet);
 		}
 	}
diff --git a/Engine/Source/Editor/CurveEditor/Private/DragOperations/CurveEditorDragOperation_MoveKeys.cpp b/Engine/Source/Editor/CurveEditor/Private/DragOperations/CurveEditorDragOperation_MoveKeys.cpp
--- a/Engine/Source/Editor/CurveEditor/Private/DragOperations/CurveEditorDragOperation_MoveKeys.cpp
+++ b/Engine/Source/Editor/CurveEditor/Private/DragOperations/CurveEditorDragOperation_MoveKeys.cpp
@@ -112,54 +112,6 @@ void FCurveEditorDragOperation_MoveKeys::OnEndDrag(FVector2D InitialPosition, FV
 	{
 		if (FCurveModel* Curve = CurveEditor->FindCurve(KeyData.CurveID))
 		{
-			// Gather all the final key time positions
-			TArray<FKeyPosition> KeyTimes;
-			for (int32 KeyIndex = 0; KeyIndex < KeyData.Handles.Num(); ++KeyIndex)
-			{
-				const FKeyHandle& KeyHandle = KeyData.Handles[KeyIndex];
-				FKeyPosition KeyTime = KeyData.LastDraggedKeyPositions[KeyIndex];
-				KeyTimes.Add(KeyTime);
-			}
-
-			// For each key time, look for all the keys that match
-			TArray<FKeyHandle> KeysToRemove;
-			for (const FKeyPosition& KeyTime : KeyTimes)
-			{
-				TArray<FKeyHandle> KeysInRange;
-				Curve->GetKeys(KeyTime.InputValue, KeyTime.InputValue, TNumericLimits<double>::Lowest(), TNumericLimits<double>::Max(), KeysInRange);
-
-				// If there's more than 1 key at this time, remove all but the keys that moved the largest amount
-				if (KeysInRange.Num() > 1)
-				{
-					double MaxDist = TNumericLimits<double>::Min();
-					FKeyHandle MaxKeyHandle;
-					for (const FKeyHandle& KeyInRange : KeysInRange)
-					{
-						int32 KeyDataIndex = KeyData.Handles.Find(KeyInRange);
-						if (KeyDataIndex != INDEX_NONE)
-						{
-							double Dist = FMath::Abs(KeyData.StartKeyPositions[KeyDataIndex].InputValue - KeyData.LastDraggedKeyPositions[KeyDataIndex].InputValue);
-							if (Dist > MaxDist)
-							{
-								MaxKeyHandle = KeyInRange;
-								MaxDist = Dist;
-							}
-						}
-					}
-
-					for (const FKeyHandle& KeyInRange : KeysInRange)
-					{
-						if (KeyInRange != MaxKeyHandle)
-						{
-							KeysToRemove.Add(KeyInRange);
-						}
-					}
-				}
-			}
-
-			// Remove any keys that overlap before moving new keys on top
-			Curve->RemoveKeys(KeysToRemove, 0.0);
-			// Then, move the keys
 			Curve->SetKeyPositions(KeyData.Handles, KeyData.LastDraggedKeyPositions, EPropertyChangeType::ValueSet);
 		}
 	}
@@ -256,4 +208,4 @@ void FCurveEditorDragOperation_MoveKeys::UpdateFromDrag(const FVector2D& Initial
 		// Make sure the last dragged key positions are up to date
 		Curve->GetKeyPositions(KeyData.Handles, KeyData.LastDraggedKeyPositions);
 	}
-}
\ No newline at end of file
+}
diff --git a/Engine/Source/Editor/MovieSceneTools/Private/Channels/ChannelCurveModel.cpp b/Engine/Source/Editor/MovieSceneTools/Private/Channels/ChannelCurveModel.cpp
--- a/Engine/Source/Editor/MovieSceneTools/Private/Channels/ChannelCurveModel.cpp
+++ b/Engine/Source/Editor/MovieSceneTools/Private/Channels/ChannelCurveModel.cpp
@@ -312,7 +312,8 @@ void FChannelCurveModel<ChannelType, ChannelValue, KeyType>::SetKeyPositions(TAr
 			{
 				FFrameNumber NewTime = (InKeyPositions[Index].InputValue * TickResolution).RoundToFrame() - KeyOffset;
 
-				KeyIndex = ChannelData.MoveKey(KeyIndex, NewTime);
+				const bool bRemoveDuplicates = ChangeType == EPropertyChangeType::ValueSet;
+				KeyIndex = ChannelData.MoveKey(KeyIndex, NewTime, bRemoveDuplicates);
 				SetKeyValue(KeyIndex, InKeyPositions[Index].OutputValue);
 			}
 		}
diff --git a/Engine/Source/Editor/MovieSceneTools/Private/Channels/ConstraintChannelCurveModel.cpp b/Engine/Source/Editor/MovieSceneTools/Private/Channels/ConstraintChannelCurveModel.cpp
--- a/Engine/Source/Editor/MovieSceneTools/Private/Channels/ConstraintChannelCurveModel.cpp
+++ b/Engine/Source/Editor/MovieSceneTools/Private/Channels/ConstraintChannelCurveModel.cpp
@@ -264,7 +264,8 @@ void FConstraintChannelCurveModel::SetKeyPositions(TArrayView<const FKeyHandle>
 			{
 				FFrameNumber NewTime = (InKeyPositions[Index].InputValue * TickResolution).FloorToFrame();
 
-				KeyIndex = ChannelData.MoveKey(KeyIndex, NewTime);
+				const bool bRemoveDuplicates = ChangeType == EPropertyChangeType::ValueSet;
+				KeyIndex = ChannelData.MoveKey(KeyIndex, NewTime, bRemoveDuplicates);
 				Section->ExpandToFrame(NewTime);
 			}
 		}
diff --git a/Engine/Source/Editor/MovieSceneTools/Private/Channels/EventChannelCurveModel.cpp b/Engine/Source/Editor/MovieSceneTools/Private/Channels/EventChannelCurveModel.cpp
--- a/Engine/Source/Editor/MovieSceneTools/Private/Channels/EventChannelCurveModel.cpp
+++ b/Engine/Source/Editor/MovieSceneTools/Private/Channels/EventChannelCurveModel.cpp
@@ -198,7 +198,8 @@ void FEventChannelCurveModel::SetKeyPositions(TArrayView<const FKeyHandle> InKey
 			{
 				FFrameNumber NewTime = (InKeyPositions[Index].InputValue * TickResolution).FloorToFrame();
 
-				KeyIndex = ChannelData.MoveKey(KeyIndex, NewTime);
+				const bool bRemoveDuplicates = ChangeType == EPropertyChangeType::ValueSet;
+				KeyIndex = ChannelData.MoveKey(KeyIndex, NewTime, bRemoveDuplicates);
 				Section->ExpandToFrame(NewTime);
 			}
 		}
diff --git a/Engine/Source/Editor/Sequencer/Private/MVVM/Selection.cpp b/Engine/Source/Editor/Sequencer/Private/MVVM/Selection.cpp
--- a/Engine/Source/Editor/Sequencer/Private/MVVM/Selection.cpp
+++ b/Engine/Source/Editor/Sequencer/Private/MVVM/Selection.cpp
@@ -9,11 +9,59 @@
 #include "MVVM/Extensions/ITrackExtension.h"
 #include "MVVM/Extensions/ISelectableExtension.h"
 #include "MVVM/ViewModels/TrackRowModel.h"
+#include "Channels/MovieSceneChannel.h"
+#include "MovieSceneSection.h"
 #include "MovieSceneTrack.h"
 
 namespace UE::Sequencer
 {
 
+void FKeySelection::RemoveDuplicateKeys(TArrayView<FKeyHandle> KeyHandles)
+{
+	TMap<TViewModelPtr<FChannelModel>, TArray<FKeyHandle> > ChannelToKeyHandles;
+
+	for (const FKeyHandle& KeyHandle : KeyHandles)
+	{
+		TViewModelPtr<FChannelModel> ChannelModel = GetModelForKey(KeyHandle);
+		if (ChannelModel)
+		{
+			if (FMovieSceneChannel* Channel = ChannelModel->GetChannel())
+			{
+				FFrameNumber KeyTime;
+				Channel->GetKeyTime(KeyHandle, KeyTime);
+
+				TArray<FFrameNumber> OutKeyTimes;
+				TArray<FKeyHandle> OutKeyHandles;
+				Channel->GetKeys(TRange<FFrameNumber>(KeyTime, KeyTime), &OutKeyTimes, &OutKeyHandles);
+				OutKeyHandles.Remove(KeyHandle);
+
+				ChannelToKeyHandles.FindOrAdd(ChannelModel).Append(OutKeyHandles);
+			}
+		}
+	}
+
+	for (TPair<TViewModelPtr<FChannelModel>, TArray<FKeyHandle> >& ChannelToKey : ChannelToKeyHandles)
+	{
+		TViewModelPtr<FChannelModel> ChannelModel = ChannelToKey.Key;
+		ChannelModel->GetSection()->Modify();
+		ChannelModel->GetChannel()->DeleteKeys(ChannelToKey.Value);
+	}
+}
+
+void FKeySelection::Deselect(FKeyHandle InKey)
+{
+	RemoveDuplicateKeys(MakeArrayView(&InKey, 1));
+
+	TUniqueFragmentSelectionSet<FKeyHandle, FChannelModel>::Deselect(InKey);
+}
+
+void FKeySelection::Empty()
+{
+	RemoveDuplicateKeys(MakeArrayView(GetSelected().Array()));
+
+	TUniqueFragmentSelectionSet<FKeyHandle, FChannelModel>::Empty();
+}
+
 bool FTrackAreaSelection::OnSelectItem(const FWeakViewModelPtr& WeakViewModel)
 {
 	if (TSharedPtr<FViewModel> ViewModel = WeakViewModel.Pin())
diff --git a/Engine/Source/Editor/Sequencer/Public/MVVM/Selection/Selection.h b/Engine/Source/Editor/Sequencer/Public/MVVM/Selection/Selection.h
--- a/Engine/Source/Editor/Sequencer/Public/MVVM/Selection/Selection.h
+++ b/Engine/Source/Editor/Sequencer/Public/MVVM/Selection/Selection.h
@@ -41,6 +41,19 @@ struct FKeySelection : TUniqueFragmentSelectionSet<FKeyHandle, FChannelModel>
 	{
 		TUniqueFragmentSelectionSet<FKeyHandle, FChannelModel>::SelectRange(InOwner, Forward<RangeType>(Range), OutAnySelected);
 	}
+
+	/**
+	 * Overridden Deselect function to remove duplicate keys on deselection
+	 */
+	void Deselect(FKeyHandle InKey);
+
+	/**
+	 * Overridden Empty function to remove duplicate keys on deselection
+	 */
+	void Empty();
+
+private:
+	void RemoveDuplicateKeys(TArrayView<FKeyHandle> KeyHandles);
 };
 
 
diff --git a/Engine/Source/Runtime/MovieScene/Public/Channels/MovieSceneChannelData.h b/Engine/Source/Runtime/MovieScene/Public/Channels/MovieSceneChannelData.h
--- a/Engine/Source/Runtime/MovieScene/Public/Channels/MovieSceneChannelData.h
+++ b/Engine/Source/Runtime/MovieScene/Public/Channels/MovieSceneChannelData.h
@@ -335,8 +335,31 @@ struct TMovieSceneChannelData : FMovieSceneChannelData
 	 * @param NewTime        The time to move the key to
 	 * @return The index of the key in its new position
 	 */
-	int32 MoveKey(int32 KeyIndex, FFrameNumber NewTime)
+
+	int32 MoveKey(int32 KeyIndex, FFrameNumber NewTime, bool bRemoveDuplicateKeys = false)
 	{
+		if (bRemoveDuplicateKeys)
+		{
+			FKeyHandle KeyHandle = GetHandle(KeyIndex);
+
+			TArray<FKeyHandle> KeysToRemove;
+			TArray<FFrameNumber> DummyTimes;
+			GetKeys(TRange<FFrameNumber>(NewTime, NewTime), &DummyTimes, &KeysToRemove);
+			KeysToRemove.Remove(KeyHandle);
+
+			if (KeysToRemove.Num() > 0)
+			{
+				DeleteKeys(KeysToRemove);
+			}
+
+			KeyIndex = GetIndex(KeyHandle);
+		}
+
+		if (KeyIndex == INDEX_NONE)
+		{
+			return KeyIndex;
+		}
+
 		int32 NewIndex = MoveKeyInternal(KeyIndex, NewTime);
 		if (NewIndex != KeyIndex)
 		{
@@ -346,6 +369,7 @@ struct TMovieSceneChannelData : FMovieSceneChannelData
 			Values->RemoveAt(KeyIndex, EAllowShrinking::No);
 			Values->Insert(OldValue, NewIndex);
 		}
+
 		return NewIndex;
 	}
 
