Curve Editor: Change curve draw params cache to exist per view instead of per curve editor, fix related issues where constraint and space bars could no longer be edited and showed curves inside

#rnx

diff --git a/Engine/Plugins/Animation/ControlRig/Source/ControlRigEditor/Private/ControlRigEditorModule.cpp b/Engine/Plugins/Animation/ControlRig/Source/ControlRigEditor/Private/ControlRigEditorModule.cpp
--- a/Engine/Plugins/Animation/ControlRig/Source/ControlRigEditor/Private/ControlRigEditorModule.cpp
+++ b/Engine/Plugins/Animation/ControlRig/Source/ControlRigEditor/Private/ControlRigEditorModule.cpp
@@ -330,8 +330,7 @@ void FControlRigEditorModule::StartupModule()
 	FControlRigSpaceChannelCurveModel::ViewID = CurveEditorModule.RegisterView(FOnCreateCurveEditorView::CreateStatic(
 		[](TWeakPtr<FCurveEditor> WeakCurveEditor) -> TSharedRef<SCurveEditorView>
 		{
-			return SNew(SCurveEditorKeyBarView, WeakCurveEditor)
-				.DrawCurves(false);
+			return SNew(SCurveEditorKeyBarView, WeakCurveEditor);
 		}
 	));
 
diff --git a/Engine/Source/Editor/CurveEditor/Private/CurveEditor.cpp b/Engine/Source/Editor/CurveEditor/Private/CurveEditor.cpp
--- a/Engine/Source/Editor/CurveEditor/Private/CurveEditor.cpp
+++ b/Engine/Source/Editor/CurveEditor/Private/CurveEditor.cpp
@@ -7,7 +7,6 @@
 #include "CoreGlobals.h"
 #include "CurveEditorCommands.h"
 #include "CurveEditorCopyBuffer.h"
-#include "CurveEditorCurveDrawParamsCache.h"
 #include "CurveEditorSettings.h"
 #include "CurveEditorSnapMetrics.h"
 #include "CurveEditorAxis.h"
@@ -133,8 +132,6 @@ void FCurveEditor::InitCurveEditor(const FCurveEditorInitParams& InInitParams)
 	SuspendBroadcastCount = 0;
 	// Listen to global undo so we can fix up our selection state for keys that no longer exist.
 	GEditor->RegisterForUndo(this);
-
-	CurveDrawParamsCache = MakeShared<UE::CurveEditor::FCurveDrawParamsCache>(AsShared());
 }
 
 int32 FCurveEditor::GetSupportedTangentTypes()
diff --git a/Engine/Source/Editor/CurveEditor/Private/CurveEditorCurveDrawParamsCache.cpp b/Engine/Source/Editor/CurveEditor/Private/CurveEditorCurveDrawParamsCache.cpp
--- a/Engine/Source/Editor/CurveEditor/Private/CurveEditorCurveDrawParamsCache.cpp
+++ b/Engine/Source/Editor/CurveEditor/Private/CurveEditorCurveDrawParamsCache.cpp
@@ -24,6 +24,8 @@ namespace UE::CurveEditor
 	{
 		TRACE_CPUPROFILER_EVENT_SCOPE(FCurveDrawParamsCache::Invalidate);
 
+		WeakCurveEditor = CurveEditorView->GetCurveEditor();
+
 		UpdateCurveCacheFlags(CurveEditorView);
 		DrawCurves(CurveEditorView, ModelIDs);
 	}
@@ -199,8 +201,7 @@ namespace UE::CurveEditor
 		}
 	}
 
-	FCurveDrawParamsCache::FCurveDrawParamsCache(const TSharedRef<FCurveEditor>& InCurveEditor)
-		: WeakCurveEditor(InCurveEditor)
+	FCurveDrawParamsCache::FCurveDrawParamsCache()
 	{
 		CurveCacheFlags = ECurveCacheFlags::UpdateAll;
 		CachedCurveEditorData.ActiveCurvesSerialNumber = 0xFFFFFFFF;
diff --git a/Engine/Source/Editor/CurveEditor/Private/SCurveEditorView.cpp b/Engine/Source/Editor/CurveEditor/Private/SCurveEditorView.cpp
--- a/Engine/Source/Editor/CurveEditor/Private/SCurveEditorView.cpp
+++ b/Engine/Source/Editor/CurveEditor/Private/SCurveEditorView.cpp
@@ -49,6 +49,8 @@ SCurveEditorView::SCurveEditorView()
 	CachedValues.CachedGeometrySize.X = -1;
 	CachedValues.CachedGeometrySize.Y = -1;
 	PRAGMA_ENABLE_DEPRECATION_WARNINGS
+
+	CurveDrawParamsCache = MakeShared<UE::CurveEditor::FCurveDrawParamsCache>();
 }
 
 PRAGMA_DISABLE_DEPRECATION_WARNINGS // Allow deprecated CachedCurveParams to destruct
@@ -98,8 +100,8 @@ void SCurveEditorView::UpdateCurveViewTransformsFromModels()
 
 	if (bTransformChanged)
 	{
-		CurveEditor->GetCurveDrawParamsCache().Invalidate(SharedThis(this));
-		
+		CurveDrawParamsCache->Invalidate(SharedThis(this));
+
 		PRAGMA_DISABLE_DEPRECATION_WARNINGS
 		CurveCacheFlags = ECurveCacheFlags::All;
 		PRAGMA_ENABLE_DEPRECATION_WARNINGS
@@ -443,6 +445,11 @@ void SCurveEditorView::FrameHorizontal(double InInputMin, double InInputMax, FCu
 	}
 }
 
+TSharedPtr<FCurveEditor> SCurveEditorView::GetCurveEditor() const
+{
+	return WeakCurveEditor.IsValid() ? WeakCurveEditor.Pin() : nullptr;
+}
+
 void SCurveEditorView::SetOutputBounds(double InOutputMin, double InOutputMax, FCurveEditorViewAxisID AxisID)
 {
 	if (!bFixedOutputBounds && InOutputMin < InOutputMax)
@@ -770,7 +777,7 @@ void SCurveEditorView::CheckCacheAndInvalidateIfNeeded()
 		TArray<FCurveModelID> ModelIDs;
 		CurveInfoByID.GetKeys(ModelIDs);
 
-		CurveEditor->GetCurveDrawParamsCache().Invalidate(SharedThis(this), ModelIDs);
+		CurveDrawParamsCache->Invalidate(SharedThis(this), ModelIDs);
 	}
 	else
 	{
@@ -782,9 +789,9 @@ void SCurveEditorView::CheckCacheAndInvalidateIfNeeded()
 		GetInputBounds(CachedValues.CachedInputMin, CachedValues.CachedInputMax);
 		CachedValues.CachedGeometrySize = GetCachedGeometry().GetLocalSize();
 	
-		CurveEditor->GetCurveDrawParamsCache().CachedDrawParams.Reset();
+		CurveDrawParamsCache->CachedDrawParams.Reset();
 
-		GetCurveDrawParams(CurveEditor->GetCurveDrawParamsCache().CachedDrawParams);
+		GetCurveDrawParams(CurveDrawParamsCache->CachedDrawParams);
 		PRAGMA_ENABLE_DEPRECATION_WARNINGS
 
 		RefreshRetainer();
diff --git a/Engine/Source/Editor/CurveEditor/Private/SCurveViewerPanel.cpp b/Engine/Source/Editor/CurveEditor/Private/SCurveViewerPanel.cpp
--- a/Engine/Source/Editor/CurveEditor/Private/SCurveViewerPanel.cpp
+++ b/Engine/Source/Editor/CurveEditor/Private/SCurveViewerPanel.cpp
@@ -40,7 +40,7 @@ void SCurveViewerPanel::Construct(const FArguments& InArgs, TSharedRef<FCurveEdi
 
 	InCurveEditor->SetView(SharedThis(this));
 
-	InCurveEditor->GetCurveDrawParamsCache().Invalidate(SharedThis(this));
+	CurveDrawParamsCache->Invalidate(SharedThis(this));
 
 	for (const TPair<FCurveModelID, TUniquePtr<FCurveModel>>& CurvePair : InCurveEditor->GetCurves())
 	{
diff --git a/Engine/Source/Editor/CurveEditor/Private/Views/SInteractiveCurveEditorView.cpp b/Engine/Source/Editor/CurveEditor/Private/Views/SInteractiveCurveEditorView.cpp
--- a/Engine/Source/Editor/CurveEditor/Private/Views/SInteractiveCurveEditorView.cpp
+++ b/Engine/Source/Editor/CurveEditor/Private/Views/SInteractiveCurveEditorView.cpp
@@ -123,7 +123,7 @@ void SInteractiveCurveEditorView::Construct(const FArguments& InArgs, TWeakPtr<F
 		TSharedPtr<FCurveEditor> CurveEditorPin = InCurveEditor.Pin();
 		check(CurveEditorPin);
 
-		CurveEditorPin->GetCurveDrawParamsCache().Invalidate(SharedThis(this));
+		CurveDrawParamsCache->Invalidate(SharedThis(this));
 		
 		CurveEditorPin->OnActiveToolChangedDelegate.AddSP(this, &SInteractiveCurveEditorView::OnCurveEditorToolChanged);
 		CurveEditorPin->GetSettings()->OnShowValueIndicatorsChanged().AddSP(this, &SInteractiveCurveEditorView::OnShowValueIndicatorsChanged);
@@ -799,7 +799,7 @@ void SInteractiveCurveEditorView::DrawCurves(TSharedRef<FCurveEditor> CurveEdito
 	const bool bDrawLines = CVarDrawCurveLines.GetValueOnGameThread();
 	const bool bDrawKeys = CVarDrawCurveKeys.GetValueOnGameThread();
 
-	const TArray<FCurveDrawParams>& DrawParamsArray = CurveEditor->GetCurveDrawParamsCache().GetCurveDrawParams();
+	const TArray<FCurveDrawParams>& DrawParamsArray = CurveDrawParamsCache->GetCurveDrawParams();
 
 	TOptional<FCurveModelID> HoveredCurve = GetHoveredCurve();
 	for (const FCurveDrawParams& DrawParams : DrawParamsArray)
@@ -1098,7 +1098,7 @@ bool SInteractiveCurveEditorView::GetPointsWithinWidgetRange(const FSlateRect& W
 
 	const float PointOverlapSensitivity = CurveEditor->GetSettings()->GetMarqueePointSensitivity();
 
-	const TArray<FCurveDrawParams>& DrawParamsArray = CurveEditor->GetCurveDrawParamsCache().GetCurveDrawParams();
+	const TArray<FCurveDrawParams>& DrawParamsArray = CurveDrawParamsCache->GetCurveDrawParams();
 
 	// Iterate through all of our points and see which points the marquee overlaps. Both of these coordinate systems
 	// are in screen space pixels.  Also check tangent lines
@@ -1149,7 +1149,7 @@ bool SInteractiveCurveEditorView::GetCurveWithinWidgetRange(const FSlateRect& Wi
 
 	FBox WidgetRectangleBox(FVector(WidgetRectangle.Left, WidgetRectangle.Top, 0), FVector(WidgetRectangle.Right, WidgetRectangle.Bottom, 0));
 
-	const TArray<FCurveDrawParams>& DrawParamsArray =CurveEditor->GetCurveDrawParamsCache().GetCurveDrawParams();
+	const TArray<FCurveDrawParams>& DrawParamsArray = CurveDrawParamsCache->GetCurveDrawParams();
 
 	// Iterate through all of our interpolating points and terminates if one overlaps the marquee. Both of these coordinate systems
 	// are in screen space pixels.
@@ -1915,7 +1915,7 @@ TOptional<FCurvePointHandle> SInteractiveCurveEditorView::HitPoint(FVector2D Mou
 
 	TOptional<FCurveModelID> HoveredCurve = GetHoveredCurve();
 
-	const TArray<FCurveDrawParams>& DrawParamsArray = CurveEditor->GetCurveDrawParamsCache().GetCurveDrawParams();
+	const TArray<FCurveDrawParams>& DrawParamsArray = CurveDrawParamsCache->GetCurveDrawParams();
 
 	// Find all keys within the current hit test time
 	for (const FCurveDrawParams& DrawParams : DrawParamsArray)
diff --git a/Engine/Source/Editor/CurveEditor/Public/CurveEditor.h b/Engine/Source/Editor/CurveEditor/Public/CurveEditor.h
--- a/Engine/Source/Editor/CurveEditor/Public/CurveEditor.h
+++ b/Engine/Source/Editor/CurveEditor/Public/CurveEditor.h
@@ -53,7 +53,6 @@ struct FCurveDrawParams;
 struct FCurveEditorInitParams;
 struct FCurveEditorSnapMetrics;
 struct FGeometry;
-namespace UE::CurveEditor { class FCurveDrawParamsCache; }
 
 DECLARE_DELEGATE_OneParam(FOnSetBoolean, bool)
 DECLARE_MULTICAST_DELEGATE_OneParam(FOnActiveToolChanged, FCurveEditorToolID)
@@ -359,9 +358,6 @@ public:
 	/** Returns whether the buffered curve is to be acted on, ie. selected, in the tree view or with selected keys */
 	bool IsActiveBufferedCurve(const TUniquePtr<IBufferedCurveModel>& BufferedCurve) const;
 
-	/** Returns a cache for curve draw params */
-	UE::CurveEditor::FCurveDrawParamsCache& GetCurveDrawParamsCache() { return *CurveDrawParamsCache; }
-
 	// ~FCurveEditor
 
 	// FEditorUndoClient
@@ -704,9 +700,6 @@ protected:
 	int32 SuspendBroadcastCount;
 
 private:
-	/** Cache for curve draw params */
-	TSharedPtr<UE::CurveEditor::FCurveDrawParamsCache> CurveDrawParamsCache;
-
 	/** Cached physical size of the panel representing this editor */
 	FVector2D CachedPhysicalSize;
 
diff --git a/Engine/Source/Editor/CurveEditor/Public/CurveEditorCurveDrawParamsCache.h b/Engine/Source/Editor/CurveEditor/Public/CurveEditorCurveDrawParamsCache.h
--- a/Engine/Source/Editor/CurveEditor/Public/CurveEditorCurveDrawParamsCache.h
+++ b/Engine/Source/Editor/CurveEditor/Public/CurveEditorCurveDrawParamsCache.h
@@ -64,7 +64,7 @@ namespace UE::CurveEditor
 			FVector2D CachedGeometrySize = FVector2D::ZeroVector;
 		};
 
-		FCurveDrawParamsCache(const TSharedRef<FCurveEditor>& InCurveEditor);
+		FCurveDrawParamsCache();
 
 		/** 
 		 * Invalidates the cache for specified view and curve model IDs and causes an update on cached data.
@@ -110,6 +110,6 @@ namespace UE::CurveEditor
 		FCachedCurveEditorData CachedCurveEditorData;
 
 		/** The curve editor that owns this cache */
-		const TWeakPtr<FCurveEditor> WeakCurveEditor;
+		TWeakPtr<FCurveEditor> WeakCurveEditor;
 	};
 }
diff --git a/Engine/Source/Editor/CurveEditor/Public/SCurveEditorView.h b/Engine/Source/Editor/CurveEditor/Public/SCurveEditorView.h
--- a/Engine/Source/Editor/CurveEditor/Public/SCurveEditorView.h
+++ b/Engine/Source/Editor/CurveEditor/Public/SCurveEditorView.h
@@ -30,6 +30,7 @@ class SRetainerWidget;
 class SCurveEditorView;
 
 enum class ECurveEditorAxisOrientation : uint8;
+namespace UE::CurveEditor { class FCurveDrawParamsCache; }
 
 /**
  * Identifier for a specific axis on a view. These identifiers are transient and should not be stored externally.
@@ -253,6 +254,9 @@ public:
 	/** Frame the view horizontall by the input and output bounds, peformaing any custom clipping as needed */
 	virtual void FrameHorizontal(double InInputMin, double InInputMax, FCurveEditorViewAxisID AxisID = FCurveEditorViewAxisID());
 
+	/** Returns the cure editor associated with this view, or nullptr if the view is not specific to a curve editor */
+	TSharedPtr<FCurveEditor> GetCurveEditor() const;
+
 public:
 
 	/**
@@ -334,11 +338,11 @@ public:
 
 protected:
 	/** Gets info about the curves being drawn. Converts actual curves into an abstract series of lines/points/handles/etc. */
-	UE_DEPRECATED(5.6, "Instead  get the cache by calling FCurveEditor::GetCurveDrawParamsCache, use FCurveDrawParamsCache::GetCurveDrawParamsSynchronous to get Curve Draw Parms as with this function.")
+	UE_DEPRECATED(5.6, "Instead use the CurveDrawParamsCache, use FCurveDrawParamsCache::GetCurveDrawParamsSynchronous to get Curve Draw Parms as with this function.")
 	void GetCurveDrawParams(TArray<FCurveDrawParams>& OutDrawParams);
 
 	/** Gets curve draw params  get the cache by calling a single curve */
-	UE_DEPRECATED(5.6, "Instead use FCurveEditor::GetCurveDrawParamsCache, use FCurveDrawParamsCache::GetCurveDrawParamSynchronous to get Curve Draw Parm as with this function.")
+	UE_DEPRECATED(5.6, "Instead use GetCurveDrawParamsCache, use FCurveDrawParamsCache::GetCurveDrawParamSynchronous to get Curve Draw Parm as with this function.")
 	void GetCurveDrawParam(TSharedPtr<FCurveEditor>& CurveEditor, const FCurveModelID& ModelID, FCurveModel* CurveModel, FCurveDrawParams& OutDrawParam) const;
 
 	/** Update all the curve view transforms from curve models */
@@ -444,24 +448,27 @@ protected:
 
 	/** Flag enum signifying how the curve cache has changed since it was last generated
 	Note for a data change it may only effect certain data(curves) not every drawn curve*/	
-	enum class UE_DEPRECATED(5.6, "Instead  get the cache by calling FCurveEditor::GetCurveDrawParamsCache, call FCurveDrawParamsCache::ECurveCacheFlags to get a corresponding type.") ECurveCacheFlags : uint8
+	enum class UE_DEPRECATED(5.6, "Instead use the CurveDrawParamsCache,  call FCurveDrawParamsCache::ECurveCacheFlags to get a corresponding type.") ECurveCacheFlags : uint8
 	{
 		CheckCurves = 0,       // The cache may be valid need to check each curve to see if they are still valid
 		All = 1 << 0,		   // Get all
 	};
 
 	/** Curve cache flags that change based upon data or view getting modified*/
-	UE_DEPRECATED(5.6, "Instead get the cache by calling FCurveEditor::GetCurveDrawParamsCache, call FCurveDrawParamsCache::GetCurveCacheFlags to get current flags. Use FCurveDrawParamsCache::Update to update them.")
+	UE_DEPRECATED(5.6, "Instead use the CurveDrawParamsCache, call FCurveDrawParamsCache::GetCurveCacheFlags to get current flags. Use FCurveDrawParamsCache::Update to update them.")
 	PRAGMA_DISABLE_DEPRECATION_WARNINGS
 	ECurveCacheFlags CurveCacheFlags;
 	PRAGMA_ENABLE_DEPRECATION_WARNINGS
 
 	/** Curve draw parameters that are re-generated on tick if the cache has changed. We generate them once and then they're used in multiple places per frame. */
-	UE_DEPRECATED(5.5, "Instead  get the cache by calling FCurveEditor::GetCurveDrawParamsCache, call GetCurveDrawParamsSynchronous.")
+	UE_DEPRECATED(5.5, "Instead use the CurveDrawParamsCache, call GetCurveDrawParamsSynchronous.")
 	TArray<FCurveDrawParams> CachedDrawParams;
 
+	/** Cache for curve draw params */
+	TSharedPtr<UE::CurveEditor::FCurveDrawParamsCache> CurveDrawParamsCache;
+
 	/** Set of Cached values we need to check each tick to see if we need to redo cache*/
-	struct UE_DEPRECATED(5.6, "Instead  get the cache by calling FCurveEditor::GetCurveDrawParamsCache, call FCurveDrawParamsCache::FCachedCurveEditorData to get a corresponding type.") FCachedValuesToCheck
+	struct UE_DEPRECATED(5.6, "Instead use the CurveDrawParamsCache, call FCurveDrawParamsCache::FCachedCurveEditorData to get a corresponding type.") FCachedValuesToCheck
 	{
 		/** Serial number cached from FCurveEditor::GetActiveCurvesSerialNumber() on tick */
 		uint32 CachedActiveCurvesSerialNumber;
@@ -479,11 +486,11 @@ protected:
 		FVector2D CachedGeometrySize;
 	};
 
-	UE_DEPRECATED(5.6, "Instead  get the cache by calling FCurveEditor::GetCurveDrawParamsCache, use FCurveDrawParamsCache::GetCurveEditorData to get cached values. Use FCurveDrawParamsCache::Update to update them.")
+	UE_DEPRECATED(5.6, "Instead use the CurveDrawParamsCache, use FCurveDrawParamsCache::GetCurveEditorData to get cached values. Use FCurveDrawParamsCache::Update to update them.")
 	PRAGMA_DISABLE_DEPRECATION_WARNINGS
 	FCachedValuesToCheck CachedValues;
 	PRAGMA_ENABLE_DEPRECATION_WARNINGS
-	
+
 	/** Possible pointer to a retainer widget that we may need to force update*/
 	TSharedPtr<SRetainerWidget> RetainerWidget;
 
diff --git a/Engine/Source/Editor/MovieSceneTools/Private/Channels/SCurveEditorKeyBarView.cpp b/Engine/Source/Editor/MovieSceneTools/Private/Channels/SCurveEditorKeyBarView.cpp
--- a/Engine/Source/Editor/MovieSceneTools/Private/Channels/SCurveEditorKeyBarView.cpp
+++ b/Engine/Source/Editor/MovieSceneTools/Private/Channels/SCurveEditorKeyBarView.cpp
@@ -39,8 +39,6 @@ float SCurveEditorKeyBarView::TrackHeight = 24.f;
 
 void SCurveEditorKeyBarView::Construct(const FArguments& InArgs, TWeakPtr<FCurveEditor> InCurveEditor)
 {
-	bDrawCurves = InArgs._DrawCurves;
-
 	bFixedOutputBounds = true;
 	bAllowModelViewTransforms = false;
 
@@ -49,7 +47,7 @@ void SCurveEditorKeyBarView::Construct(const FArguments& InArgs, TWeakPtr<FCurve
 	WeakCurveEditor = InCurveEditor;
 	SortBias = 25;
 
-	SInteractiveCurveEditorView::Construct(SInteractiveCurveEditorView::FArguments(), InCurveEditor);
+	SInteractiveCurveEditorView::Construct(InArgs, InCurveEditor);
 }
 
 FVector2D SCurveEditorKeyBarView::ComputeDesiredSize(float LayoutScaleMultiplier) const
@@ -94,11 +92,7 @@ void SCurveEditorKeyBarView::PaintView(const FPaintArgs& Args, const FGeometry&
 		DrawBackground(AllottedGeometry, OutDrawElements, BaseLayerId, DrawEffects);
 		DrawLabels(AllottedGeometry, OutDrawElements, BaseLayerId, DrawEffects);
 		DrawGridLines(CurveEditor.ToSharedRef(), AllottedGeometry, OutDrawElements, BaseLayerId, DrawEffects);
-		
-		if (bDrawCurves)
-		{
-			DrawCurves(CurveEditor.ToSharedRef(), AllottedGeometry, MyCullingRect, OutDrawElements, BaseLayerId, InWidgetStyle, DrawEffects);
-		}
+		DrawCurves(CurveEditor.ToSharedRef(), AllottedGeometry, MyCullingRect, OutDrawElements, BaseLayerId, InWidgetStyle, DrawEffects);
 	}
 }
 
diff --git a/Engine/Source/Editor/MovieSceneTools/Private/MovieSceneToolsModule.cpp b/Engine/Source/Editor/MovieSceneTools/Private/MovieSceneToolsModule.cpp
--- a/Engine/Source/Editor/MovieSceneTools/Private/MovieSceneToolsModule.cpp
+++ b/Engine/Source/Editor/MovieSceneTools/Private/MovieSceneToolsModule.cpp
@@ -263,8 +263,7 @@ void FMovieSceneToolsModule::StartupModule()
 		FConstraintChannelCurveModel::ViewID = CurveEditorModule.RegisterView(FOnCreateCurveEditorView::CreateStatic(
 			[](TWeakPtr<FCurveEditor> WeakCurveEditor) -> TSharedRef<SCurveEditorView>
 			{
-				return SNew(SCurveEditorKeyBarView, WeakCurveEditor)
-					.DrawCurves(false);
+				return SNew(SCurveEditorKeyBarView, WeakCurveEditor);
 			}
 		));
 
diff --git a/Engine/Source/Editor/MovieSceneTools/Public/Channels/SCurveEditorKeyBarView.h b/Engine/Source/Editor/MovieSceneTools/Public/Channels/SCurveEditorKeyBarView.h
--- a/Engine/Source/Editor/MovieSceneTools/Public/Channels/SCurveEditorKeyBarView.h
+++ b/Engine/Source/Editor/MovieSceneTools/Public/Channels/SCurveEditorKeyBarView.h
@@ -24,15 +24,6 @@ struct FGeometry;
 class MOVIESCENETOOLS_API SCurveEditorKeyBarView : public SInteractiveCurveEditorView
 {
 public:
-	SLATE_BEGIN_ARGS(SCurveEditorKeyBarView) 
-		: _DrawCurves(true)
-		{}
-	
-		/** When false the widget will not draw curves or keys */
-		SLATE_ARGUMENT(bool, DrawCurves)
-	
-	SLATE_END_ARGS()
-
 	void Construct(const FArguments& InArgs, TWeakPtr<FCurveEditor> InCurveEditor);
 
 	virtual void GetGridLinesY(TSharedRef<const FCurveEditor> CurveEditor, TArray<float>& MajorGridLines, TArray<float>& MinorGridLines, TArray<FText>* MajorGridLabels) const override;
@@ -50,8 +41,5 @@ private:
 
 	void DrawLabels(const FGeometry& AllottedGeometry, FSlateWindowElementList& OutDrawElements, int32 BaseLayerId, ESlateDrawEffect DrawEffects) const;
 
-	/** When false the widget will not draw curves or keys */
-	bool bDrawCurves;
-
 	static float TrackHeight;
 };
\ No newline at end of file
