Curve Editor: Introduce command based undo / redo system. Allows key selection to be undo / redoable.

Changes required to address fallout:
- FCurveEditorSelection::RequestSyncSelection now keeps transaction alive until next tick
- New: KeySelection::FAddInternal: Passkey pattern
- Adjusted several places to use the new FScopedSelectionTransaction instead of FScopedTransaction and added it to some places that did not have it before (to support selection undo)

diff --git a/Engine/Source/Editor/CurveEditor/Private/CurveEditor.cpp b/Engine/Source/Editor/CurveEditor/Private/CurveEditor.cpp
--- a/Engine/Source/Editor/CurveEditor/Private/CurveEditor.cpp
+++ b/Engine/Source/Editor/CurveEditor/Private/CurveEditor.cpp
@@ -23,6 +23,7 @@
 #include "Framework/Commands/UICommandList.h"
 #include "Framework/Notifications/NotificationManager.h"
 #include "Framework/SlateDelegates.h"
+#include "HAL/IConsoleManager.h"
 #include "HAL/PlatformApplicationMisc.h"
 #include "ICurveEditorExtension.h"
 #include "ICurveEditorModule.h"
@@ -57,6 +58,7 @@
 #include "UObject/UnrealNames.h"
 #include "UnrealExporter.h"
 #include "Misc/SmartSnap.h"
+#include "Modification/Utils/ScopedSelectionTransaction.h"
 #include "Widgets/Colors/SColorPicker.h"
 #include "Widgets/Notifications/SNotificationList.h"
 
@@ -135,6 +137,7 @@ void FCurveEditor::InitCurveEditor(const FCurveEditorInitParams& InInitParams)
 	GEditor->RegisterForUndo(this);
 
 	CurveDrawParamsCache = MakeShared<UE::CurveEditor::FCurveDrawParamsCache>(AsShared());
+	TransactionManager = MakeUnique<UE::CurveEditor::FTransactionManager>(SharedThis(this));
 }
 
 int32 FCurveEditor::GetSupportedTangentTypes()
@@ -210,17 +213,29 @@ void FCurveEditor::ClearAxes()
 
 FCurveModelID FCurveEditor::AddCurve(TUniquePtr<FCurveModel>&& InCurve)
 {
-	FCurveModelID NewID = FCurveModelID::Unique();
+	check(InCurve);
+	
+	// The curve ID is relevant e.g. for undo / redo.
+	// You can undo / redo selecting keys: if undo past a transaction that called FCurveEditor::AddCurve, redoing that transaction needs to add back
+	// the same curve ID so redoing the key selection also works.
+	// If InCurve has no ID set, GetOrInitId will set it here. If the caller actually specifies an ID and that ID is added already, they have a
+	// mistake in their business logic.
+	const FCurveModelID& CurveId = InCurve->GetOrInitId();
+	if (!ensureMsgf(!CurveData.Contains(CurveId), TEXT("Investigate what caused the double-addition and fix it!")))
+	{
+		return FCurveModelID();
+	}
+	
 	FCurveModel *Curve = InCurve.Get();
 
-	CurveData.Add(NewID, MoveTemp(InCurve));
+	CurveData.Add(CurveId, MoveTemp(InCurve));
 
 	// Add child curves
 	TArray<TUniquePtr<FCurveModel>> ChildCurvesArray;
 	Curve->MakeChildCurves(ChildCurvesArray);
 	for (TUniquePtr<FCurveModel>& Child : ChildCurvesArray)
 	{
-		ChildCurves.Add(NewID, AddCurve(MoveTemp(Child)));
+		ChildCurves.Add(CurveId, AddCurve(MoveTemp(Child)));
 	}
 
 	++ActiveCurvesSerialNumber;
@@ -228,7 +243,7 @@ FCurveModelID FCurveEditor::AddCurve(TUniquePtr<FCurveModel>&& InCurve)
 	{
 		OnCurveArrayChanged.Broadcast(Curve, true, this);
 	}
-	return NewID;
+	return CurveId;
 }
 
 void FCurveEditor::BroadcastCurveChanged(FCurveModel* InCurve)
@@ -1198,6 +1213,7 @@ void FCurveEditor::ClearSelectionRange()
 
 void FCurveEditor::SelectAllKeys()
 {
+	const UE::CurveEditor::FScopedSelectionTransaction Transaction(SharedThis(this), LOCTEXT("SelectAllKeys", "Select all keys"));
 	for (FCurveModelID ID : GetEditedCurves())
 	{
 		if (FCurveModel* Curve = FindCurve(ID))
@@ -1211,6 +1227,7 @@ void FCurveEditor::SelectAllKeys()
 
 void FCurveEditor::SelectForward()
 {
+	const UE::CurveEditor::FScopedSelectionTransaction Transaction(SharedThis(this), LOCTEXT("SelectForward", "Select forward"));
 	Selection.Clear();
 
 	TSharedPtr<ITimeSliderController> TimeSliderController = WeakTimeSliderController.Pin();
@@ -1236,6 +1253,7 @@ void FCurveEditor::SelectForward()
 
 void FCurveEditor::SelectBackward()
 {
+	const UE::CurveEditor::FScopedSelectionTransaction Transaction(SharedThis(this), LOCTEXT("SelectBackward", "Select backward"));
 	Selection.Clear();
 
 	TSharedPtr<ITimeSliderController> TimeSliderController = WeakTimeSliderController.Pin();
@@ -1261,11 +1279,13 @@ void FCurveEditor::SelectBackward()
 
 void FCurveEditor::SelectNone()
 {
+	const UE::CurveEditor::FScopedSelectionTransaction Transaction(SharedThis(this));
 	Selection.Clear();
 }
 
 void FCurveEditor::InvertSelection()
 {
+	const UE::CurveEditor::FScopedSelectionTransaction Transaction(SharedThis(this), LOCTEXT("InvertSelection", "Invert selection"));
 	for (const TTuple<FCurveModelID, FKeyHandleSet>& Pair : Selection.GetAll())
 	{
 		FCurveModelID CurveModelID = Pair.Key;
@@ -1771,7 +1791,7 @@ void FCurveEditor::PasteKeys(TSet<FCurveModelID> CurveModelIDs, const bool bInOv
 		}
 	}
 
-	FScopedTransaction Transaction(LOCTEXT("PasteKeys", "Paste Keys"));
+	const UE::CurveEditor::FScopedSelectionTransaction KeyChange(SharedThis(this), LOCTEXT("PasteKeys", "Paste Keys"));
 	Selection.Clear();
 
 
@@ -1902,7 +1922,7 @@ void FCurveEditor::PasteKeys(TSet<FCurveModelID> CurveModelIDs, const bool bInOv
 
 void FCurveEditor::DeleteSelection()
 {
-	FScopedTransaction Transaction(LOCTEXT("DeleteKeys", "Delete Keys"));
+	const UE::CurveEditor::FScopedSelectionTransaction Transaction(SharedThis(this), LOCTEXT("DeleteKeys", "Delete Keys"));
 
 	double CurrentTime = 0.0;
 	if (TSharedPtr<ITimeSliderController> TimeSliderController = WeakTimeSliderController.Pin())
@@ -2069,7 +2089,7 @@ bool FCurveEditor::CanFlattenOrStraightenSelection() const
 
 void FCurveEditor::SmartSnapSelection()
 {
-	FScopedTransaction Transaction(LOCTEXT("SmartSnapKeys", "Smart Snap"));
+	const UE::CurveEditor::FScopedSelectionTransaction Transaction(SharedThis(this), LOCTEXT("SmartSnapKeys", "Smart Snap"));
 	
 	TMap<FCurveModelID, FKeyHandleSet> OutKeysToSelect;
 	UE::CurveEditor::EnumerateSmartSnappableKeys(*this, Selection.GetAll(), OutKeysToSelect,
@@ -2508,6 +2528,10 @@ bool FCurveEditor::IsActiveBufferedCurve(const TUniquePtr<IBufferedCurveModel>&
 	return false;
 }
 
+static TAutoConsoleVariable<bool> CVarDisableKeyCleansing(
+	TEXT("CurveEditor.EnableCurveCleansing"), false, TEXT("After undo operations, whether to remove invalid keys from the key selection.")
+	);
+
 void FCurveEditor::PostUndo(bool bSuccess)
 {
 	if (WeakPanel.IsValid())
@@ -2515,6 +2539,15 @@ void FCurveEditor::PostUndo(bool bSuccess)
 		WeakPanel.Pin()->PostUndo();
 	}
 
+	// Temporary hack while we're testing the new undo / redo system for key selection.
+	// Disabled by default - if animators notice problems, they can enable this CVar. We don't expect issues with this off by default.
+	// If no issues are found in 2 weeks, we'll just remove the entirety of the below code.
+	// 19th of Feb 2025. 
+	if (!CVarDisableKeyCleansing.GetValueOnGameThread())
+	{
+		return;
+	}
+
 	// If you create keys and then undo them the selection set still thinks there's keys selected.
 	// This presents issues with context menus and other things that are activated when there is a selection set.
 	// To fix this, we have to loop through all of our curve models, and re-select only the key handles that were
diff --git a/Engine/Source/Editor/CurveEditor/Private/CurveEditorSelection.cpp b/Engine/Source/Editor/CurveEditor/Private/CurveEditorSelection.cpp
--- a/Engine/Source/Editor/CurveEditor/Private/CurveEditorSelection.cpp
+++ b/Engine/Source/Editor/CurveEditor/Private/CurveEditorSelection.cpp
@@ -33,6 +33,17 @@ void FCurveEditorSelection::SerializeSelection(FArchive& Ar)
 	Ar << CurveToSelectedKeys;
 }
 
+void FCurveEditorSelection::AddInternal(
+	FCurveModelID CurveID, ECurvePointType PointType, TArrayView<const FKeyHandle> Keys, bool bIncrementSerialNumber
+	)
+{
+	FKeyHandleSet& SelectedKeys = CurveToSelectedKeys.FindOrAdd(CurveID);
+	for (const FKeyHandle& Key : Keys)
+	{
+		SelectedKeys.Add(Key, PointType);
+	}
+}
+
 const FKeyHandleSet* FCurveEditorSelection::FindForCurve(FCurveModelID InCurveID) const
 {
 	return CurveToSelectedKeys.Find(InCurveID);
@@ -78,12 +89,8 @@ void FCurveEditorSelection::Add(FCurveModelID CurveID, ECurvePointType PointType
 		const TUniquePtr<FCurveModel> *CurveModel = CurveEditor->GetCurves().Find(CurveID);
 		if (CurveModel && CurveModel->IsValid() && !(*CurveModel)->IsReadOnly())
 		{
-			FKeyHandleSet& SelectedKeys = CurveToSelectedKeys.FindOrAdd(CurveID);
-			for (FKeyHandle Key : Keys)
-			{
-				SelectedKeys.Add(Key, PointType);
-			}
-
+			constexpr bool bIncrementSerialNumber = false;
+			AddInternal(CurveID, PointType, Keys, bIncrementSerialNumber);
 		}
 	}
 
diff --git a/Engine/Source/Editor/CurveEditor/Private/CurveEditorSelectionPrivate.h b/Engine/Source/Editor/CurveEditor/Private/CurveEditorSelectionPrivate.h
new file mode 100644
--- /dev/null
+++ b/Engine/Source/Editor/CurveEditor/Private/CurveEditorSelectionPrivate.h
@@ -0,0 +1,23 @@
+﻿// Copyright Epic Games, Inc. All Rights Reserved.
+
+#pragma once
+
+#include "CurveEditorSelection.h"
+#include "Modification/Utils/SelectionChangeUtils.h"
+
+namespace UE::CurveEditor::KeySelection 
+{
+/** Implements passkey pattern to restrict access to FCurveEditorSelection::AddInternal. */
+class FAddInternal
+{
+	friend void SelectionDiffDetail::AddKeysInternal(FCurveEditorSelection&, FCurveModelID, ECurvePointType, TArrayView<const FKeyHandle>);
+	
+	explicit FAddInternal(
+		FCurveEditorSelection& Selection, FCurveModelID CurveID, ECurvePointType PointType, TArrayView<const FKeyHandle> Keys,
+		bool bIncrementSerialNumber = true
+		)
+	{
+		Selection.AddInternal(CurveID, PointType, Keys, bIncrementSerialNumber);
+	}
+};
+}
diff --git a/Engine/Source/Editor/CurveEditor/Private/DragOperations/CurveEditorDragOperation_Marquee.cpp b/Engine/Source/Editor/CurveEditor/Private/DragOperations/CurveEditorDragOperation_Marquee.cpp
--- a/Engine/Source/Editor/CurveEditor/Private/DragOperations/CurveEditorDragOperation_Marquee.cpp
+++ b/Engine/Source/Editor/CurveEditor/Private/DragOperations/CurveEditorDragOperation_Marquee.cpp
@@ -151,16 +151,14 @@ static FSelectionPreferences GetSelectionPreferences(
 }
 
 FCurveEditorDragOperation_Marquee::FCurveEditorDragOperation_Marquee(FCurveEditor* InCurveEditor)
-	: LockedToView(nullptr)
-{
-	CurveEditor = InCurveEditor;
-}
+	: FCurveEditorDragOperation_Marquee(InCurveEditor, nullptr)
+{}
 
 FCurveEditorDragOperation_Marquee::FCurveEditorDragOperation_Marquee(FCurveEditor* InCurveEditor, SCurveEditorView* InLockedToView)
-	: LockedToView(InLockedToView)
-{
-	CurveEditor = InCurveEditor;
-}
+	: SelectionChange(InCurveEditor->AsShared())
+	, CurveEditor(InCurveEditor)
+	, LockedToView(InLockedToView)
+{}
 
 void FCurveEditorDragOperation_Marquee::OnBeginDrag(FVector2D InitialPosition, FVector2D CurrentPosition, const FPointerEvent& MouseEvent)
 {
diff --git a/Engine/Source/Editor/CurveEditor/Private/DragOperations/CurveEditorDragOperation_Marquee.h b/Engine/Source/Editor/CurveEditor/Private/DragOperations/CurveEditorDragOperation_Marquee.h
--- a/Engine/Source/Editor/CurveEditor/Private/DragOperations/CurveEditorDragOperation_Marquee.h
+++ b/Engine/Source/Editor/CurveEditor/Private/DragOperations/CurveEditorDragOperation_Marquee.h
@@ -7,6 +7,7 @@
 #include "ICurveEditorDragOperation.h"
 #include "Layout/SlateRect.h"
 #include "Math/Vector2D.h"
+#include "Modification/Utils/ScopedSelectionTransaction.h"
 
 class FCurveEditor;
 class FSlateWindowElementList;
@@ -28,6 +29,9 @@ public:
 
 private:
 
+	/** Makes this change undo / redo able. */
+	const UE::CurveEditor::FScopedSelectionTransaction SelectionChange;
+
 	/** The current marquee rectangle */
 	FSlateRect Marquee;
 	/** Ptr back to the curve editor */
diff --git a/Engine/Source/Editor/CurveEditor/Private/Filters/FilterUtils.cpp b/Engine/Source/Editor/CurveEditor/Private/Filters/FilterUtils.cpp
--- a/Engine/Source/Editor/CurveEditor/Private/Filters/FilterUtils.cpp
+++ b/Engine/Source/Editor/CurveEditor/Private/Filters/FilterUtils.cpp
@@ -3,6 +3,7 @@
 #include "FilterUtils.h"
 
 #include "Filters/CurveEditorFilterBase.h"
+#include "Modification/Utils/ScopedSelectionTransaction.h"
 
 #include "Containers/Map.h"
 #include "Templates/SharedPointer.h"
@@ -22,7 +23,8 @@ namespace UE::CurveEditor::FilterUtils
 	{
 		TMap<FCurveModelID, FKeyHandleSet> OutKeysToSelect;
 		InFilter.ApplyFilter(InCurveEditor, InSelectedKeys, OutKeysToSelect);
-
+		
+		const FScopedSelectionTransaction Transaction(InCurveEditor);
 		// Clear their selection and then set it to the keys the filter thinks you should have selected.
 		InCurveEditor->GetSelection().Clear();
 
@@ -31,4 +33,4 @@ namespace UE::CurveEditor::FilterUtils
 			InCurveEditor->GetSelection().Add(OutSet.Key, ECurvePointType::Key, OutSet.Value.AsArray());
 		}
 	}
-}
+}
\ No newline at end of file
diff --git a/Engine/Source/Editor/CurveEditor/Private/Modification/Changes/SelectionChangeCommand.cpp b/Engine/Source/Editor/CurveEditor/Private/Modification/Changes/SelectionChangeCommand.cpp
new file mode 100644
--- /dev/null
+++ b/Engine/Source/Editor/CurveEditor/Private/Modification/Changes/SelectionChangeCommand.cpp
@@ -0,0 +1,25 @@
+﻿// Copyright Epic Games, Inc. All Rights Reserved.
+
+#include "SelectionChangeCommand.h"
+
+#include "CurveEditor.h"
+#include "Modification/Utils/SelectionChangeUtils.h"
+
+namespace UE::CurveEditor
+{
+void FSelectionChangeCommand::Apply(UObject* Object)
+{
+	if (const TSharedPtr<FCurveEditor> CurveEditor = GetCurveEditor(Object))
+	{
+		ApplySelectionChange(CurveEditor->Selection, DeltaChange);
+	}
+}
+
+void FSelectionChangeCommand::Revert(UObject* Object)
+{
+	if (const TSharedPtr<FCurveEditor> CurveEditor = GetCurveEditor(Object))
+	{
+		RevertSelectionChange(CurveEditor->Selection, DeltaChange);
+	}
+}
+}
diff --git a/Engine/Source/Editor/CurveEditor/Private/Modification/Changes/SelectionChangeCommand.h b/Engine/Source/Editor/CurveEditor/Private/Modification/Changes/SelectionChangeCommand.h
new file mode 100644
--- /dev/null
+++ b/Engine/Source/Editor/CurveEditor/Private/Modification/Changes/SelectionChangeCommand.h
@@ -0,0 +1,28 @@
+﻿// Copyright Epic Games, Inc. All Rights Reserved.
+
+#pragma once
+
+#include "SelectionDeltaChange.h"
+#include "Modification/CurveEditorCommandChange.h"
+
+namespace UE::CurveEditor
+{
+/** Changes that changes the curve editor's selection. */
+class FSelectionChangeCommand : public FCurveEditorCommandChange
+{
+public:
+
+	explicit FSelectionChangeCommand(FSelectionDeltaChange InDeltaChange) : DeltaChange(MoveTemp(InDeltaChange)) {}
+	
+	//~ Begin FChange Interface
+	virtual void Apply(UObject* Object) override;
+	virtual void Revert(UObject* Object) override;
+	//~ End FChange Interface
+
+private:
+
+	/** The change made to the selection. */
+	const FSelectionDeltaChange DeltaChange;
+};
+}
+
diff --git a/Engine/Source/Editor/CurveEditor/Private/Modification/Changes/SelectionDeltaChange.h b/Engine/Source/Editor/CurveEditor/Private/Modification/Changes/SelectionDeltaChange.h
new file mode 100644
--- /dev/null
+++ b/Engine/Source/Editor/CurveEditor/Private/Modification/Changes/SelectionDeltaChange.h
@@ -0,0 +1,35 @@
+﻿// Copyright Epic Games, Inc. All Rights Reserved.
+
+#pragma once
+
+#include "Containers/Array.h"
+#include "Containers/Map.h"
+#include "CurveEditorTypes.h"
+#include "CurveModel.h"
+
+namespace UE::CurveEditor
+{
+/** Changes made to a single curve's selection */
+struct FCurveSelectionDeltaChange
+{
+	// Using TArray instead of TMap because we'll be iterating through all entries frequently. TArray is much faster for that.
+	
+	/** The keys that were added and its point type. */
+	TArray<TPair<FKeyHandle, ECurvePointType>> AddedKeys;
+	/**
+	 * The keys that were removed.
+	 * Can contain a key handle that is in AddedKeys. In that case, the point type was changed.
+	 */
+	TArray<TPair<FKeyHandle, ECurvePointType>> RemovedKeys;
+};
+
+/** Changes made to the curve editor selection */
+struct FSelectionDeltaChange
+{
+	/** The per curve changes. */
+	TMap<FCurveModelID, FCurveSelectionDeltaChange> ChangedCurves;
+
+	bool HasChanges() const { return !ChangedCurves.IsEmpty(); }
+	operator bool() const { return HasChanges(); }
+};
+}
diff --git a/Engine/Source/Editor/CurveEditor/Private/Modification/CurveEditorCommandChange.cpp b/Engine/Source/Editor/CurveEditor/Private/Modification/CurveEditorCommandChange.cpp
new file mode 100644
--- /dev/null
+++ b/Engine/Source/Editor/CurveEditor/Private/Modification/CurveEditorCommandChange.cpp
@@ -0,0 +1,25 @@
+﻿// Copyright Epic Games, Inc. All Rights Reserved.
+
+#include "Modification/CurveEditorCommandChange.h"
+
+#include "Containers/UnrealString.h"
+#include "CurveEditorTransactionObject.h"
+
+namespace UE::CurveEditor
+{
+bool FCurveEditorCommandChange::HasExpired(UObject* InObject) const
+{
+	return !GetCurveEditor(InObject).IsValid();
+}
+
+FString FCurveEditorCommandChange::ToString() const
+{
+	return TEXT("FSelectionChangeCommand");
+}
+
+TSharedPtr<FCurveEditor> FCurveEditorCommandChange::GetCurveEditor(UObject* InObject)
+{
+	UCurveEditorTransactionObject* TransactionObject = Cast<UCurveEditorTransactionObject>(InObject);
+	return ensure(TransactionObject) ? TransactionObject->OwningCurveEditor.Pin() : nullptr;
+}
+}
diff --git a/Engine/Source/Editor/CurveEditor/Private/Modification/CurveEditorTransactionObject.h b/Engine/Source/Editor/CurveEditor/Private/Modification/CurveEditorTransactionObject.h
new file mode 100644
--- /dev/null
+++ b/Engine/Source/Editor/CurveEditor/Private/Modification/CurveEditorTransactionObject.h
@@ -0,0 +1,25 @@
+﻿// Copyright Epic Games, Inc. All Rights Reserved.
+
+#pragma once
+
+#include "UObject/Object.h"
+#include "Templates/SharedPointer.h"
+#include "CurveEditorTransactionObject.generated.h"
+
+class FCurveEditor;
+
+/**
+ * The object type that the transaction system associates FCommandChanges for the FCurveEditor with.
+ *
+ * By design, the transaction system requires a UObject to associate FCommandChanges with. This is usually an asset.
+ * Conceptually, instances of this class will act as an "asset" and the associated FCurveEditor contains the "asset data", i.e. its  FCurveModels.
+ */
+UCLASS()
+class CURVEEDITOR_API UCurveEditorTransactionObject : public UObject
+{
+	GENERATED_BODY()
+public:
+
+	/** The curve editor this object is associated with. Used by commands to retrieve the */
+	TWeakPtr<FCurveEditor> OwningCurveEditor;
+};
diff --git a/Engine/Source/Editor/CurveEditor/Private/Modification/TransactionManager.cpp b/Engine/Source/Editor/CurveEditor/Private/Modification/TransactionManager.cpp
new file mode 100644
--- /dev/null
+++ b/Engine/Source/Editor/CurveEditor/Private/Modification/TransactionManager.cpp
@@ -0,0 +1,36 @@
+﻿// Copyright Epic Games, Inc. All Rights Reserved.
+
+#include "Modification/TransactionManager.h"
+
+#include "CurveEditorTransactionObject.h"
+#include "Misc/ITransaction.h"
+#include "UObject/Package.h"
+#include "UObject/UObjectGlobals.h"
+
+namespace UE::CurveEditor
+{
+FTransactionManager::FTransactionManager(TWeakPtr<FCurveEditor> InCurveEditor)
+	: DummyTransactionObject(NewObject<UCurveEditorTransactionObject>(GetTransientPackage(), UCurveEditorTransactionObject::StaticClass(), NAME_None, RF_Transient | RF_Transactional))
+{
+	DummyTransactionObject->OwningCurveEditor = MoveTemp(InCurveEditor);
+}
+
+void FTransactionManager::AppendChange(TUniquePtr<FCurveEditorCommandChange> InChange) const
+{
+	if (ensureAlways(GUndo) && ensure(DummyTransactionObject) && ensure(InChange))
+	{
+		GUndo->StoreUndo(DummyTransactionObject, MoveTemp(InChange));
+	}
+}
+
+void FTransactionManager::AddReferencedObjects(FReferenceCollector& Collector)
+{
+	Collector.AddReferencedObject(DummyTransactionObject);
+}
+
+FString FTransactionManager::GetReferencerName() const
+{
+	return TEXT("FTransactionManager");
+}
+}
+
diff --git a/Engine/Source/Editor/CurveEditor/Private/Modification/Utils/ScopedSelectionTransaction.cpp b/Engine/Source/Editor/CurveEditor/Private/Modification/Utils/ScopedSelectionTransaction.cpp
new file mode 100644
--- /dev/null
+++ b/Engine/Source/Editor/CurveEditor/Private/Modification/Utils/ScopedSelectionTransaction.cpp
@@ -0,0 +1,50 @@
+﻿// Copyright Epic Games, Inc. All Rights Reserved.
+
+#include "Modification/Utils/ScopedSelectionTransaction.h"
+
+#include "CurveEditor.h"
+#include "Internationalization/Internationalization.h"
+#include "Modification/Changes/SelectionChangeCommand.h"
+#include "SelectionChangeUtils.h"
+
+#define LOCTEXT_NAMESPACE "FScopedSelectionChange"
+
+namespace UE::CurveEditor
+{
+FScopedSelectionTransaction::FScopedSelectionTransaction(TWeakPtr<FCurveEditor> InCurveEditor, bool bInActuallyTransact)
+	: FScopedSelectionTransaction(MoveTemp(InCurveEditor), LOCTEXT("ChangeSelection", "Change selection"), bInActuallyTransact)
+{}
+
+FScopedSelectionTransaction::FScopedSelectionTransaction(TWeakPtr<FCurveEditor> InCurveEditor, const FText& InDescription, bool bInActuallyTransact)
+	: FScopedTransaction(InDescription, bInActuallyTransact)
+	, WeakCurveEditor(MoveTemp(InCurveEditor))
+{
+	if (const TSharedPtr<FCurveEditor> CurveEditorPin = WeakCurveEditor.Pin())
+	{
+		OriginalSelection = CurveEditorPin->Selection;
+	}
+}
+
+FScopedSelectionTransaction::~FScopedSelectionTransaction()
+{
+	if (!IsOutstanding())
+	{
+		return;
+	}
+	
+	const TSharedPtr<FCurveEditor> CurveEditorPin = WeakCurveEditor.Pin();
+	if (!CurveEditorPin)
+	{
+		return;
+	}
+	
+	if (FSelectionDeltaChange DeltaChange = DiffSelection(OriginalSelection, CurveEditorPin->Selection))
+	{
+		CurveEditorPin->GetTransactionManager()->AppendChange(
+			MakeUnique<FSelectionChangeCommand>(MoveTemp(DeltaChange))
+			);
+	}
+}
+}
+
+#undef LOCTEXT_NAMESPACE
\ No newline at end of file
diff --git a/Engine/Source/Editor/CurveEditor/Private/Modification/Utils/SelectionChangeUtils.cpp b/Engine/Source/Editor/CurveEditor/Private/Modification/Utils/SelectionChangeUtils.cpp
new file mode 100644
--- /dev/null
+++ b/Engine/Source/Editor/CurveEditor/Private/Modification/Utils/SelectionChangeUtils.cpp
@@ -0,0 +1,133 @@
+﻿// Copyright Epic Games, Inc. All Rights Reserved.
+
+#include "SelectionChangeUtils.h"
+
+#include "CurveEditorSelection.h"
+#include "CurveEditorSelectionPrivate.h"
+#include "Modification/Changes/SelectionDeltaChange.h"
+
+#include <type_traits>
+
+namespace UE::CurveEditor
+{
+namespace SelectionDiffDetail
+{
+/** Goes through every key in InOriginalSelection and sees which keys are missing from InTarget. The missing keys are added to the array InGetKeys returns. */
+template<typename TGetKeyArray>
+requires std::is_invocable_r_v<TArray<TPair<FKeyHandle, ECurvePointType>>&, TGetKeyArray, FCurveSelectionDeltaChange&>
+static void DetectMissingKeys(
+	const TMap<FCurveModelID, FKeyHandleSet>& InOriginalSelection, const FCurveEditorSelection& InTarget,
+	FSelectionDeltaChange& OutDeltaChange, TGetKeyArray&& InGetKeys
+	)
+{
+	for (const TPair<FCurveModelID, FKeyHandleSet>& OriginalCurvePair : InOriginalSelection)
+	{
+		const FCurveModelID& CurveId = OriginalCurvePair.Key;
+		const FKeyHandleSet& OriginalKeySet = OriginalCurvePair.Value; 
+		const FKeyHandleSet* TargetKeySet = InTarget.FindForCurve(CurveId);
+
+		// Should not really happen, but we'll handle the case anyway.
+		if (OriginalKeySet.Num() == 0 && (!TargetKeySet || TargetKeySet->Num() == 0))
+		{
+			continue;
+		}	
+
+		// Fully missing from target?
+		if (!TargetKeySet
+			// We'd expect that FindForCurve would have returned nullptr if Num() == 0 but we'll handle the case anyway.
+			|| TargetKeySet->Num() == 0)
+		{
+			FCurveSelectionDeltaChange& CurveChange = OutDeltaChange.ChangedCurves.Add(CurveId);
+			TArray<TPair<FKeyHandle, ECurvePointType>>& Keys = InGetKeys(CurveChange);
+			Keys.Reserve(OriginalKeySet.Num());
+			for (const FKeyHandle& RemovedKey : OriginalKeySet.AsArray())
+			{
+				Keys.Emplace(RemovedKey, OriginalKeySet.PointType(RemovedKey));
+			}
+			continue;
+		}
+
+		// Find keys that were are missing or changed point type
+		for (const FKeyHandle& OriginalKey : OriginalKeySet.AsArray())
+		{
+			// If the key handle is still in there but the point type has changed, we'll add it back again below.
+			const ECurvePointType PointType = OriginalKeySet.PointType(OriginalKey);
+			if (!TargetKeySet->Contains(OriginalKey, PointType))
+			{
+				TArray<TPair<FKeyHandle, ECurvePointType>>& Keys = InGetKeys(OutDeltaChange.ChangedCurves.FindOrAdd(CurveId));
+				Keys.Emplace(OriginalKey, PointType);
+			}
+		}
+	}
+}
+	
+template<typename TCallback>
+concept CGetKeys = std::is_invocable_r_v<const TArray<TPair<FKeyHandle, ECurvePointType>>&, TCallback, const FCurveSelectionDeltaChange&>;
+
+void AddKeysInternal(FCurveEditorSelection& Selection, FCurveModelID CurveID, ECurvePointType PointType, TArrayView<const FKeyHandle> Keys)
+{
+	KeySelection::FAddInternal(Selection, CurveID, PointType, Keys);
+}
+	
+/** Removes keys from the array InGetRemovedKeys returns, and adds keys to the array InGetAddedKeys returns. */
+template<CGetKeys TGetAddedKeys, CGetKeys TGetRemovedKeys>
+static void ModifySelection(
+	FCurveEditorSelection& InOutChange, const FSelectionDeltaChange& InDeltaChange, TGetAddedKeys&& InGetAddedKeys, TGetRemovedKeys&& InGetRemovedKeys
+	)
+{
+	for (const TPair<FCurveModelID, FCurveSelectionDeltaChange>& ChangePair : InDeltaChange.ChangedCurves)
+	{
+		const FCurveModelID& CurveId = ChangePair.Key;
+		const TArray<TPair<FKeyHandle, ECurvePointType>>& AddedKeys = InGetAddedKeys(ChangePair.Value);
+		const TArray<TPair<FKeyHandle, ECurvePointType>>& RemovedKeys = InGetRemovedKeys(ChangePair.Value);
+	
+		for (const TPair<FKeyHandle, ECurvePointType>& RemovedPair : RemovedKeys)
+		{
+			InOutChange.Remove(CurveId, RemovedPair.Value, RemovedPair.Key);
+		}
+		for (const TPair<FKeyHandle, ECurvePointType>& AddedPair : AddedKeys)
+		{
+			const TConstArrayView<FKeyHandle> Keys(&AddedPair.Key, 1);
+			AddKeysInternal(InOutChange, CurveId, AddedPair.Value, Keys);
+		}
+	}
+}
+}
+	
+FSelectionDeltaChange DiffSelection(const FCurveEditorSelection& InOriginal, const FCurveEditorSelection& InTarget)
+{
+	FSelectionDeltaChange DeltaChange;
+	const TMap<FCurveModelID, FKeyHandleSet>& OriginalSelection = InOriginal.GetAll();
+	const TMap<FCurveModelID, FKeyHandleSet>& TargetSelection = InTarget.GetAll();
+	
+	// Go through everything in InOriginal and see what's missing in InTarget -> effectively computes what was removed.
+	SelectionDiffDetail::DetectMissingKeys(OriginalSelection, InTarget, DeltaChange,
+		[](FCurveSelectionDeltaChange& CurveChanges) -> TArray<TPair<FKeyHandle, ECurvePointType>>&
+		{ return CurveChanges.RemovedKeys; }
+		);
+	// Go through everything in InTarget and see what's missing in InOriginal -> effectively computes what was added.
+	SelectionDiffDetail::DetectMissingKeys(TargetSelection, InOriginal, DeltaChange,
+		[](FCurveSelectionDeltaChange& CurveChanges) -> TArray<TPair<FKeyHandle, ECurvePointType>>&
+		{ return CurveChanges.AddedKeys; }
+	);
+	
+	return DeltaChange;
+}
+
+void ApplySelectionChange(FCurveEditorSelection& InOutChange, const FSelectionDeltaChange& InDeltaChange)
+{
+	SelectionDiffDetail::ModifySelection(InOutChange, InDeltaChange,
+		[](const FCurveSelectionDeltaChange& CurveChange){ return CurveChange.AddedKeys; },
+		[](const FCurveSelectionDeltaChange& CurveChange){ return CurveChange.RemovedKeys; }
+		);
+}
+	
+void RevertSelectionChange(FCurveEditorSelection& InOutChange, const FSelectionDeltaChange& InDeltaChange)
+{
+	SelectionDiffDetail::ModifySelection(InOutChange, InDeltaChange,
+		// Swapping added and removed keys here effectively reverts the change.
+		[](const FCurveSelectionDeltaChange& CurveChange){ return CurveChange.RemovedKeys; },
+		[](const FCurveSelectionDeltaChange& CurveChange){ return CurveChange.AddedKeys; }
+		);
+}
+}
diff --git a/Engine/Source/Editor/CurveEditor/Private/Modification/Utils/SelectionChangeUtils.h b/Engine/Source/Editor/CurveEditor/Private/Modification/Utils/SelectionChangeUtils.h
new file mode 100644
--- /dev/null
+++ b/Engine/Source/Editor/CurveEditor/Private/Modification/Utils/SelectionChangeUtils.h
@@ -0,0 +1,26 @@
+﻿// Copyright Epic Games, Inc. All Rights Reserved.
+
+#pragma once
+
+#include "Containers/ArrayView.h"
+
+enum class ECurvePointType : uint8;
+struct FCurveModelID;
+struct FCurveEditorSelection;
+struct FKeyHandle;
+
+namespace UE::CurveEditor
+{
+struct FSelectionDeltaChange;
+
+/** Computes the delta change to get from InOriginal to InTarget. */
+FSelectionDeltaChange DiffSelection(const FCurveEditorSelection& InOriginal, const FCurveEditorSelection& InTarget);
+
+/** Applies InDeltaChange to InOutSelection (redo operation). */
+void ApplySelectionChange(FCurveEditorSelection& InOutChange, const FSelectionDeltaChange& InDeltaChange);
+/** Reverts InDeltaChange from InOutSelection (undo operation).  */
+void RevertSelectionChange(FCurveEditorSelection& InOutChange, const FSelectionDeltaChange& InDeltaChange);
+
+// Private - needed for FAddInternal.
+namespace SelectionDiffDetail { void AddKeysInternal(FCurveEditorSelection&, FCurveModelID, ECurvePointType, TArrayView<const FKeyHandle>); }
+}
diff --git a/Engine/Source/Editor/CurveEditor/Private/SCurveEditorPanel.cpp b/Engine/Source/Editor/CurveEditor/Private/SCurveEditorPanel.cpp
--- a/Engine/Source/Editor/CurveEditor/Private/SCurveEditorPanel.cpp
+++ b/Engine/Source/Editor/CurveEditor/Private/SCurveEditorPanel.cpp
@@ -66,6 +66,7 @@
 #include "SlotBase.h"
 #include "Filters/CurveEditorEulerFilter.h"
 #include "Filters/PromotedFilterContainer.h"
+#include "Modification/Utils/ScopedSelectionTransaction.h"
 #include "Styling/AppStyle.h"
 #include "Styling/CoreStyle.h"
 #include "Styling/ISlateStyle.h"
@@ -459,6 +460,7 @@ void SCurveEditorPanel::BindCommands()
 		TSharedPtr<FCurveEditor> LocalCurveEditor = CurveEditor;
 		FExecuteAction DeselectAllAction = FExecuteAction::CreateLambda([LocalCurveEditor] { if (LocalCurveEditor.IsValid())
 			{
+				const UE::CurveEditor::FScopedSelectionTransaction Transaction(LocalCurveEditor, LOCTEXT("DeselectAllKeys", "Deselect all keys"));
 				LocalCurveEditor->GetSelection().Clear();
 			} 
 		});
@@ -1136,6 +1138,7 @@ FReply SCurveEditorPanel::OnKeyDown(const FGeometry& MyGeometry, const FKeyEvent
 {
 	if (InKeyEvent.GetKey() == EKeys::Escape)
 	{
+		const UE::CurveEditor::FScopedSelectionTransaction Transaction(GetCurveEditor(), LOCTEXT("EscapeKeys", "Press escape"));
 		CurveEditor->Selection.Clear();
 		return FReply::Handled();
 	}
diff --git a/Engine/Source/Editor/CurveEditor/Private/SCurveEditorViewContainer.cpp b/Engine/Source/Editor/CurveEditor/Private/SCurveEditorViewContainer.cpp
--- a/Engine/Source/Editor/CurveEditor/Private/SCurveEditorViewContainer.cpp
+++ b/Engine/Source/Editor/CurveEditor/Private/SCurveEditorViewContainer.cpp
@@ -370,6 +370,7 @@ FReply SCurveEditorViewContainer::OnMouseButtonUp(const FGeometry& MyGeometry, c
 	{
 		if (!MouseEvent.IsShiftDown() && !MouseEvent.IsAltDown() && !MouseEvent.IsControlDown())
 		{
+			const UE::CurveEditor::FScopedSelectionTransaction Transaction(CurveEditor, LOCTEXT("ClearSelection", "Clear selection"));
 			CurveEditor->GetSelection().Clear();
 			Reply = FReply::Handled();
 		}
diff --git a/Engine/Source/Editor/CurveEditor/Private/Tree/SCurveEditorTreeSelect.cpp b/Engine/Source/Editor/CurveEditor/Private/Tree/SCurveEditorTreeSelect.cpp
--- a/Engine/Source/Editor/CurveEditor/Private/Tree/SCurveEditorTreeSelect.cpp
+++ b/Engine/Source/Editor/CurveEditor/Private/Tree/SCurveEditorTreeSelect.cpp
@@ -13,6 +13,7 @@
 #include "GenericPlatform/GenericApplication.h"
 #include "Layout/Children.h"
 #include "Math/NumericLimits.h"
+#include "Modification/Utils/ScopedSelectionTransaction.h"
 #include "Styling/AppStyle.h"
 #include "Templates/UnrealTemplate.h"
 #include "Tree/CurveEditorTree.h"
@@ -21,6 +22,7 @@
 
 struct FSlateBrush;
 
+#define LOCTEXT_NAMESPACE "SCurveEditorTreeSelect"
 
 void SCurveEditorTreeSelect::Construct(const FArguments& InArgs, TWeakPtr<FCurveEditor> InCurveEditor, FCurveEditorTreeItemID InTreeItemID, const TSharedRef<ITableRow>& InTableRow)
 {
@@ -58,6 +60,7 @@ FReply SCurveEditorTreeSelect::SelectAll()
 	const bool bIsShiftDown = FSlateApplication::Get().GetModifierKeys().IsShiftDown();
 	const bool bIsControlDown = FSlateApplication::Get().GetModifierKeys().IsControlDown();
 
+	const UE::CurveEditor::FScopedSelectionTransaction Transaction(CurveEditor, LOCTEXT("Tree.SelectAll", "Select all"));
 	if (!bIsShiftDown && !bIsControlDown)
 	{
 		CurveEditor->GetSelection().Clear();
@@ -116,3 +119,5 @@ const FSlateBrush* SCurveEditorTreeSelect::GetSelectBrush() const
 {
 	return FAppStyle::GetBrush("GenericCurveEditor.Select");
 }
+
+#undef LOCTEXT_NAMESPACE
\ No newline at end of file
diff --git a/Engine/Source/Editor/CurveEditor/Private/Views/SInteractiveCurveEditorView.cpp b/Engine/Source/Editor/CurveEditor/Private/Views/SInteractiveCurveEditorView.cpp
--- a/Engine/Source/Editor/CurveEditor/Private/Views/SInteractiveCurveEditorView.cpp
+++ b/Engine/Source/Editor/CurveEditor/Private/Views/SInteractiveCurveEditorView.cpp
@@ -62,6 +62,7 @@
 #include "SCurveEditorPanel.h"
 #include "ScopedTransaction.h"
 #include "SlotBase.h"
+#include "Modification/Utils/ScopedSelectionTransaction.h"
 #include "Styling/AppStyle.h"
 #include "Styling/CoreStyle.h"
 #include "Styling/ISlateStyle.h"
@@ -1467,6 +1468,7 @@ double SInteractiveCurveEditorView::GetTangentValue(const double InTime, const d
 
 void SInteractiveCurveEditorView::HandleDirectKeySelectionByMouse(TSharedPtr<FCurveEditor> CurveEditor, const FPointerEvent& MouseEvent, TOptional<FCurvePointHandle> MouseDownPoint)
 {
+	const UE::CurveEditor::FScopedSelectionTransaction Transaction(WeakCurveEditor, LOCTEXT("ClickKey", "Click Key"));
 	if(!MouseDownPoint.IsSet())
 	{
 		CurveEditor->GetSelection().Clear();
@@ -1630,7 +1632,7 @@ FReply SInteractiveCurveEditorView::OnMouseButtonDown(const FGeometry& MyGeometr
 				FCurveModel* CurveToAddTo = CurveEditor->FindCurve(HoveredCurve.GetValue());
 				if (CurveToAddTo && !CurveToAddTo->IsReadOnly() && !bDraggingKeys && !bKeyAlreadyExists)
 				{
-					FScopedTransaction Transaction(LOCTEXT("InsertKey", "Insert Key"));
+					UE::CurveEditor::FScopedSelectionTransaction Transaction(CurveEditor, LOCTEXT("InsertKey", "Insert Key"));
 
 					FCurveEditorScreenSpace CurveSpace = GetCurveSpace(HoveredCurve.GetValue());
 					double MouseTime = CurveSpace.ScreenToSeconds(MousePixel.X);
@@ -1798,6 +1800,7 @@ FReply SInteractiveCurveEditorView::OnMouseButtonUp(const FGeometry& MyGeometry,
 		TOptional<FCurveModelID> HitCurve = GetHoveredCurve();
 		if (!HitPoint(MousePosition).IsSet() && HitCurve.IsSet())
 		{
+			const UE::CurveEditor::FScopedSelectionTransaction Transaction(WeakCurveEditor, LOCTEXT("ClickCurve", "Click Curve"));
 			FCurveModel* CurveModel = CurveEditor->FindCurve(HitCurve.GetValue());
 
 			TArray<FKeyHandle> KeyHandles;
@@ -2101,7 +2104,7 @@ void SInteractiveCurveEditorView::AddKeyAtTime(const TSet<FCurveModelID>& ToCurv
 		return;
 	}
 
-	FScopedTransaction Transaction(LOCTEXT("AddKeyAtTime", "Add Key"));
+	UE::CurveEditor::FScopedSelectionTransaction Transaction(WeakCurveEditor, LOCTEXT("AddKeyAtTime", "Add Key"));
 	bool bAddedKey = false;
 
 	// Clear the selection set as we will be selecting all the new keys created.
diff --git a/Engine/Source/Editor/CurveEditor/Public/CurveEditor.h b/Engine/Source/Editor/CurveEditor/Public/CurveEditor.h
--- a/Engine/Source/Editor/CurveEditor/Public/CurveEditor.h
+++ b/Engine/Source/Editor/CurveEditor/Public/CurveEditor.h
@@ -32,6 +32,7 @@
 #include "Misc/Attribute.h"
 #include "Misc/FrameRate.h"
 #include "Misc/Optional.h"
+#include "Modification//TransactionManager.h"
 #include "Templates/SharedPointer.h"
 #include "Templates/UniquePtr.h"
 #include "Tree/CurveEditorTree.h"
@@ -137,6 +138,8 @@ public:
 
 	virtual int32 GetSupportedTangentTypes();
 
+	UE::CurveEditor::FTransactionManager* GetTransactionManager() const { return TransactionManager.Get(); }
+
 public:
 
 	void SetPanel(TSharedPtr<SCurveEditorPanel> InPanel);
@@ -703,6 +706,9 @@ protected:
 	/** Counter to suspend broadcasting of changed delegates*/
 	int32 SuspendBroadcastCount;
 
+	/** Orchestrates transaction that affect this curve editor. */
+	TUniquePtr<UE::CurveEditor::FTransactionManager> TransactionManager;
+
 private:
 	/** Cache for curve draw params */
 	TSharedPtr<UE::CurveEditor::FCurveDrawParamsCache> CurveDrawParamsCache;
diff --git a/Engine/Source/Editor/CurveEditor/Public/CurveEditorSelection.h b/Engine/Source/Editor/CurveEditor/Public/CurveEditorSelection.h
--- a/Engine/Source/Editor/CurveEditor/Public/CurveEditorSelection.h
+++ b/Engine/Source/Editor/CurveEditor/Public/CurveEditorSelection.h
@@ -15,6 +15,8 @@
 #include "Templates/SharedPointer.h"
 
 class FCurveEditor;
+namespace UE::CurveEditor::KeySelection { class FAddInternal; }
+
 struct FCurvePointHandle;
 
 /**
@@ -84,6 +86,8 @@ private:
  */
 struct CURVEEDITOR_API FCurveEditorSelection
 {
+	friend UE::CurveEditor::KeySelection::FAddInternal;
+	
 	/**
 	 * Default constructor
 	 */
@@ -205,4 +209,7 @@ private:
 
 	/** A map of selected handles stored by curve ID */
 	TMap<FCurveModelID, FKeyHandleSet> CurveToSelectedKeys;
+	
+	/** Adds the specified key handles from the selection without validating whether CurveId and Keys exists in the owning curve editor. */
+	void AddInternal(FCurveModelID CurveID, ECurvePointType PointType, TArrayView<const FKeyHandle> Keys, bool bIncrementSerialNumber = true);
 };
\ No newline at end of file
diff --git a/Engine/Source/Editor/CurveEditor/Public/CurveModel.h b/Engine/Source/Editor/CurveEditor/Public/CurveModel.h
--- a/Engine/Source/Editor/CurveEditor/Public/CurveModel.h
+++ b/Engine/Source/Editor/CurveEditor/Public/CurveModel.h
@@ -571,6 +571,28 @@ public:
 		return SupportedViews;
 	}
 
+	/** Inits this curve with an optionally specified ID. */
+	void InitCurveId(FCurveModelID InCurveModelID = FCurveModelID::Unique())
+	{
+		if (ensureAlways(!CurveId))
+		{
+			CurveId = InCurveModelID;
+		}
+	}
+
+	/** @return A valid curve ID */
+	const FCurveModelID& GetOrInitId()
+	{
+		if (!CurveId)
+		{
+			InitCurveId();
+		}
+		return *CurveId;
+	}
+	
+	/** @return The curve ID - can be unset. */
+	const TOptional<FCurveModelID>& GetId() const { return CurveId; }
+
 protected:
 
 	/** This curve's short display name. Used in situations where other mechanisms provide enough context about what the curve is (such as "X") */
@@ -610,4 +632,13 @@ protected:
 
 	/** Multicast delegate broadcast on curve modification */
 	FSimpleMulticastDelegate CurveModifiedDelegate;
+
+private:
+
+	/**
+	 * This curve's unique ID.
+	 * You can set the ID when you construct the curve. If left unspecified, this is initialized by FCurveEditor when you add the curve.
+	 * The ID can only be set once during the lifetime of the FCurveModel to inconsistencies, e.g. TMap<FCurveModelId, TUniquePtr<FCurveModel>>.
+	 */
+	TOptional<FCurveModelID> CurveId;
 };
diff --git a/Engine/Source/Editor/CurveEditor/Public/Modification/CurveEditorCommandChange.h b/Engine/Source/Editor/CurveEditor/Public/Modification/CurveEditorCommandChange.h
new file mode 100644
--- /dev/null
+++ b/Engine/Source/Editor/CurveEditor/Public/Modification/CurveEditorCommandChange.h
@@ -0,0 +1,28 @@
+﻿// Copyright Epic Games, Inc. All Rights Reserved.
+
+#pragma once
+
+#include "Misc/Change.h"
+#include "Templates/SharedPointerFwd.h"
+
+class FCurveEditor;
+
+namespace UE::CurveEditor
+{
+/** Base class for actions that support undo / redo in the curve editor. */
+class CURVEEDITOR_API FCurveEditorCommandChange : public FCommandChange
+{
+public:
+
+	//~ Begin FChange Interface
+	virtual bool HasExpired(UObject* InObject) const override;
+	virtual FString ToString() const override;
+	//~ End FChange Interface
+
+protected:
+
+	/** @return The curve editor associated with Object. */
+	static TSharedPtr<FCurveEditor> GetCurveEditor(UObject* InObject);
+};
+}
+
diff --git a/Engine/Source/Editor/CurveEditor/Public/Modification/TransactionManager.h b/Engine/Source/Editor/CurveEditor/Public/Modification/TransactionManager.h
new file mode 100644
--- /dev/null
+++ b/Engine/Source/Editor/CurveEditor/Public/Modification/TransactionManager.h
@@ -0,0 +1,44 @@
+﻿// Copyright Epic Games, Inc. All Rights Reserved.
+
+#pragma once
+
+#include "CurveEditorCommandChange.h"
+#include "Templates/UnrealTemplate.h"
+#include "UObject/GCObject.h"
+#include "UObject/ObjectPtr.h"
+
+class FCurveEditor;
+class UCurveEditorTransactionObject;
+
+namespace UE::CurveEditor
+{
+/** Interacts with the transaction system to enable command-based actions for curve editor. */
+class CURVEEDITOR_API FTransactionManager : public FNoncopyable, public FGCObject
+{
+public:
+
+	explicit FTransactionManager(TWeakPtr<FCurveEditor> InCurveEditor);
+	
+	/**
+	 * Appends a change to the current transaction.
+	 *
+	 * The transaction can consist of multiple changes, e.g.
+	 * - system 1 may move keys, and
+	 * - system 2 could coniditonally snap the keys to full frames after checking that the user has toggled auto-snapping
+	 * 
+	 * @param InChange The change to associate with the currently open transaction
+	 */
+	void AppendChange(TUniquePtr<FCurveEditorCommandChange> InChange) const;
+
+	//~ Begin FGCObject Interface
+	virtual void AddReferencedObjects(FReferenceCollector& Collector) override;
+	virtual FString GetReferencerName() const override;
+	//~ End FGCObject Interface
+
+private:
+
+	/** The UObject that transactions are associated with. The transaction system needs to associate FCommandChanges with an UObject. */
+	TObjectPtr<UCurveEditorTransactionObject> DummyTransactionObject;
+};
+}
+
diff --git a/Engine/Source/Editor/CurveEditor/Public/Modification/Utils/ScopedSelectionTransaction.h b/Engine/Source/Editor/CurveEditor/Public/Modification/Utils/ScopedSelectionTransaction.h
new file mode 100644
--- /dev/null
+++ b/Engine/Source/Editor/CurveEditor/Public/Modification/Utils/ScopedSelectionTransaction.h
@@ -0,0 +1,31 @@
+﻿// Copyright Epic Games, Inc. All Rights Reserved.
+
+#pragma once
+
+#include "CurveEditorSelection.h"
+#include "ScopedTransaction.h"
+
+namespace UE::CurveEditor
+{
+/**
+ * This util starts a transaction. When the scope ends, checks whether any changes to selection were made and adds it an undo-able action.
+ * Use if you intend to change the selection.
+ */
+class CURVEEDITOR_API FScopedSelectionTransaction : public FScopedTransaction
+{
+public:
+
+	[[nodiscard]] explicit FScopedSelectionTransaction(TWeakPtr<FCurveEditor> InCurveEditor, bool bInActuallyTransact = true);
+	[[nodiscard]] explicit FScopedSelectionTransaction(TWeakPtr<FCurveEditor> InCurveEditor, const FText& InDescription, bool bInActuallyTransact = true);
+	
+	~FScopedSelectionTransaction();
+
+private:
+
+	const TWeakPtr<FCurveEditor> WeakCurveEditor;
+
+	/** The selection the editor had when the transaction was started.  */
+	FCurveEditorSelection OriginalSelection;
+};
+}
+
diff --git a/Engine/Source/Editor/Sequencer/Private/IKeyArea.cpp b/Engine/Source/Editor/Sequencer/Private/IKeyArea.cpp
--- a/Engine/Source/Editor/Sequencer/Private/IKeyArea.cpp
+++ b/Engine/Source/Editor/Sequencer/Private/IKeyArea.cpp
@@ -15,10 +15,12 @@
 #include "MovieScene.h"
 #include "MovieSceneSequence.h"
 #include "CurveEditorSettings.h"
+#include "CurveEditorTypes.h"
 
-IKeyArea::IKeyArea(TWeakPtr<ISequencerSection> InSection, FMovieSceneChannelHandle InChannel)
+IKeyArea::IKeyArea(TWeakPtr<ISequencerSection> InSection, FMovieSceneChannelHandle InChannel , const FCurveModelID& InCurveId)
 	: TreeSerialNumber(0)
 	, ChannelHandle(InChannel)
+	, CurveModelID(MakePimpl<FCurveModelID>(InCurveId))
 {
 	Reinitialize(InSection, InChannel);
 }
@@ -487,6 +489,13 @@ TUniquePtr<FCurveModel> IKeyArea::CreateCurveEditorModel(TSharedRef<ISequencer>
 					CurveModel->SetColor(SettingColor.GetValue(),false);
 				}
 			}
+
+			// Make sure the curve consistently has the same ID.
+			// If the curve is added, removed, then added again to curve editor, systems will recognize this curve as the "same".
+			// This is relevant e.g. for undo / redo.
+			// You can undo / redo selecting keys: if undo past a transaction that called FCurveEditor::AddCurve,
+			// redoing that transaction needs to add back the same curve ID so redoing the key selection also works.
+			CurveModel->InitCurveId(*CurveModelID);
 		}
 		return CurveModel;
 	}
diff --git a/Engine/Source/Editor/Sequencer/Private/MVVM/CurveEditorExtension.cpp b/Engine/Source/Editor/Sequencer/Private/MVVM/CurveEditorExtension.cpp
--- a/Engine/Source/Editor/Sequencer/Private/MVVM/CurveEditorExtension.cpp
+++ b/Engine/Source/Editor/Sequencer/Private/MVVM/CurveEditorExtension.cpp
@@ -20,6 +20,7 @@
 #include "Sequencer.h"
 #include "SequencerCommands.h"
 #include "Menus/SequencerToolbarUtils.h"
+#include "Modification/Utils/ScopedSelectionTransaction.h"
 #include "Toolkits/IToolkitHost.h"
 #include "Tree/SCurveEditorTree.h"
 #include "Tree/SCurveEditorTreeFilterStatusBar.h"
@@ -619,7 +620,12 @@ void FCurveEditorExtension::RequestSyncSelection()
 	// right around the time when we want to unload everything (such as when loading a new map in the
 	// editor). We don't want to extend the lifetime of our stuff in that case.
 	TWeakPtr<FSequencerEditorViewModel> WeakRootViewModel(WeakOwnerModel);
-	GEditor->GetTimerManager()->SetTimerForNextTick([WeakRootViewModel]()
+
+	// Key selection supports undo. If RequestSyncSelection is called as part of an ongoing transaction, record the key selection change for undo.
+	const bool bShouldRecord = GUndo != nullptr;
+	TSharedPtr<CurveEditor::FScopedSelectionTransaction> SelectionChange = MakeShared<CurveEditor::FScopedSelectionTransaction>(CurveEditorModel, bShouldRecord);
+	
+	GEditor->GetTimerManager()->SetTimerForNextTick([WeakRootViewModel, SelectionChange = MoveTemp(SelectionChange)]() mutable
 	{
 		bSyncSelectionRequested = false;
 	
@@ -639,6 +645,8 @@ void FCurveEditorExtension::RequestSyncSelection()
 		{
 			This->SyncSelection();
 		}
+
+		SelectionChange.Reset();
 	});
 }
 
diff --git a/Engine/Source/Editor/Sequencer/Private/MVVM/ViewModels/ChannelModel.cpp b/Engine/Source/Editor/Sequencer/Private/MVVM/ViewModels/ChannelModel.cpp
--- a/Engine/Source/Editor/Sequencer/Private/MVVM/ViewModels/ChannelModel.cpp
+++ b/Engine/Source/Editor/Sequencer/Private/MVVM/ViewModels/ChannelModel.cpp
@@ -143,7 +143,8 @@ UE_SEQUENCER_DEFINE_VIEW_MODEL_TYPE_ID(FDynamicChannelGroupMuteExtension);
 
 
 FChannelModel::FChannelModel(FName InChannelName, TWeakPtr<ISequencerSection> InSection, FMovieSceneChannelHandle InChannel)
-	: KeyArea(MakeShared<IKeyArea>(InSection, InChannel))
+	: CurveModelID(MakePimpl<FCurveModelID>(FCurveModelID::Unique()))
+	, KeyArea(MakeShared<IKeyArea>(InSection, InChannel, *CurveModelID))
 	, ChannelName(InChannelName)
 {
 }
@@ -166,7 +167,7 @@ void FChannelModel::Initialize(TWeakPtr<ISequencerSection> InSection, FMovieScen
 {
 	if (!KeyArea)
 	{
-		KeyArea = MakeShared<IKeyArea>(InSection, InChannel);
+		KeyArea = MakeShared<IKeyArea>(InSection, InChannel, *CurveModelID);
 	}
 	else
 	{
diff --git a/Engine/Source/Editor/Sequencer/Public/IKeyArea.h b/Engine/Source/Editor/Sequencer/Public/IKeyArea.h
--- a/Engine/Source/Editor/Sequencer/Public/IKeyArea.h
+++ b/Engine/Source/Editor/Sequencer/Public/IKeyArea.h
@@ -5,12 +5,14 @@
 #include "Math/Color.h"
 #include "UObject/ObjectKey.h"
 #include "UObject/NameTypes.h"
+#include "Templates/PimplPtr.h"
 #include "Templates/SharedPointer.h"
 #include "Containers/ArrayView.h"
 #include "Curves/KeyHandle.h"
 #include "MovieSceneCommonHelpers.h"
 #include "Channels/MovieSceneChannelHandle.h"
 
+struct FCurveModelID;
 struct FRichCurve;
 struct FKeyDrawParams;
 struct FMovieSceneChannel;
@@ -48,7 +50,7 @@ public:
 	 * @param InSection The section that owns the channel that this key area represents
 	 * @param InChannel Handle to the channel this key area represents
 	 */
-	IKeyArea(TWeakPtr<ISequencerSection> InSection, FMovieSceneChannelHandle InChannel);
+	IKeyArea(TWeakPtr<ISequencerSection> InSection, FMovieSceneChannelHandle InChannel, const FCurveModelID& InCurveId);
 
 	/**
 	 * Reinitialize this key area with a new channel
@@ -337,4 +339,7 @@ private:
 
 	/** The display text of this channel. */
 	FText DisplayText;
+
+	/** Consistent ID that is reused and injected into the curve model every time it is created. Pimpl to avoid public build.cs dependency. */
+	const TPimplPtr<FCurveModelID> CurveModelID;
 };
\ No newline at end of file
diff --git a/Engine/Source/Editor/Sequencer/Public/MVVM/ViewModels/ChannelModel.h b/Engine/Source/Editor/Sequencer/Public/MVVM/ViewModels/ChannelModel.h
--- a/Engine/Source/Editor/Sequencer/Public/MVVM/ViewModels/ChannelModel.h
+++ b/Engine/Source/Editor/Sequencer/Public/MVVM/ViewModels/ChannelModel.h
@@ -116,8 +116,11 @@ private:
 
 private:
 
+	/** Consistent ID that is reused and injected into the curve model every time it is created. Pimpl to avoid public build.cs dependency. */
+	const TPimplPtr<FCurveModelID> CurveModelID;
+	
 	TSharedPtr<IKeyArea> KeyArea;
-	FName ChannelName;
+	const FName ChannelName;
 };
 
 /**
