Curve Editor: Added horizontal/vertical curve flip feature in curve editor. Included curve flip in Niagara inline curve editor tool bar.

[RN] minor creator/niagara

diff --git a/Engine/Plugins/FX/Niagara/Source/NiagaraEditorWidgets/Private/DetailCustomizations/NiagaraDataInterfaceCurveDetails.cpp b/Engine/Plugins/FX/Niagara/Source/NiagaraEditorWidgets/Private/DetailCustomizations/NiagaraDataInterfaceCurveDetails.cpp
--- a/Engine/Plugins/FX/Niagara/Source/NiagaraEditorWidgets/Private/DetailCustomizations/NiagaraDataInterfaceCurveDetails.cpp
+++ b/Engine/Plugins/FX/Niagara/Source/NiagaraEditorWidgets/Private/DetailCustomizations/NiagaraDataInterfaceCurveDetails.cpp
@@ -1354,6 +1354,21 @@ void FNiagaraDataInterfaceCurveDetailsBase::CustomizeDetails(IDetailLayoutBuilde
 			]);
 	}
 
+	//add flip curve buttons to tool bar
+	ToolBarBuilder.AddSeparator();
+	ToolBarBuilder.AddToolBarButton(
+		FCurveEditorCommands::Get().FlipCurveHorizontal,
+		NAME_None,
+		TAttribute<FText>(),
+		TAttribute<FText>(),
+		FSlateIcon(FAppStyle::Get().GetStyleSetName(), "GenericCurveEditor.FlipCurveHorizontal"));
+	ToolBarBuilder.AddToolBarButton(
+		FCurveEditorCommands::Get().FlipCurveVertical,
+		NAME_None,
+		TAttribute<FText>(),
+		TAttribute<FText>(),
+		FSlateIcon(FAppStyle::Get().GetStyleSetName(), "GenericCurveEditor.FlipCurveVertical"));
+
 	// add templates to toolbar
 	ToolBarBuilder.AddSeparator();
 	const UNiagaraEditorSettings* Settings = GetDefault<UNiagaraEditorSettings>();
diff --git a/Engine/Source/Editor/CurveEditor/Private/CurveEditor.cpp b/Engine/Source/Editor/CurveEditor/Private/CurveEditor.cpp
--- a/Engine/Source/Editor/CurveEditor/Private/CurveEditor.cpp
+++ b/Engine/Source/Editor/CurveEditor/Private/CurveEditor.cpp
@@ -494,6 +494,10 @@ void FCurveEditor::BindCommands()
 		CommandList->MapAction(FCurveEditorCommands::Get().ToggleOutputSnapping, ToggleOutputSnapping, FCanExecuteAction(), IsOutputSnappingEnabled);
 	}
 
+	// Flip Curve
+	CommandList->MapAction(FCurveEditorCommands::Get().FlipCurveHorizontal, FExecuteAction::CreateSP(this, &FCurveEditor::FlipCurve, ECurveFlipDirection::Horizontal));
+	CommandList->MapAction(FCurveEditorCommands::Get().FlipCurveVertical, FExecuteAction::CreateSP(this, &FCurveEditor::FlipCurve, ECurveFlipDirection::Vertical));
+
 	// Flatten and Straighten Tangents
 	{
 		CommandList->MapAction(FCurveEditorCommands::Get().FlattenTangents, FExecuteAction::CreateSP(this, &FCurveEditor::FlattenSelection), FCanExecuteAction::CreateSP(this, &FCurveEditor::CanFlattenOrStraightenSelection) );
@@ -1344,6 +1348,160 @@ void FCurveEditor::ToggleOutputSnapping()
 		OnOutputSnapEnabledChanged.ExecuteIfBound(NewValue);
 	}
 }
+void FCurveEditor::FlipCurveHorizontal(TArray<FKeyPosition>& AllKeyPositions, TArray<FKeyAttributes>& AllKeyAttributes, ECurveFlipRangeType RangeType,
+											  float InRangeMin, float InRangeMax, double CurveMinTime, double CurveMaxTime)
+{
+	float RangeMin = FLT_MAX;
+	float RangeMax = -FLT_MAX;
+
+	if (RangeType == ECurveFlipRangeType::CurveRange)
+	{
+		RangeMin = FMath::Min(CurveMinTime, RangeMin);
+		RangeMax = FMath::Max(CurveMaxTime, RangeMax);
+	}
+	else
+	{
+		RangeMin = InRangeMin;
+		RangeMax = InRangeMax;
+	}
+
+	// Loop through all keys to adjust positions and tangents
+	for (int32 Index = AllKeyPositions.Num() - 1; Index >= 0; --Index)
+	{
+		FKeyPosition& Position = AllKeyPositions[Index];
+		FKeyAttributes& Attributes = AllKeyAttributes[Index];
+
+		// Mirror x value
+		Position.InputValue = RangeMax - Position.InputValue + RangeMin;
+
+		// Mirror tangent 
+		if (Attributes.HasArriveTangent() && Attributes.HasLeaveTangent())
+		{
+			float ArriveTemp = Attributes.GetArriveTangent();
+			float LeaveTemp = Attributes.GetLeaveTangent();
+			Attributes.SetArriveTangent(-LeaveTemp);
+			Attributes.SetLeaveTangent(-ArriveTemp);
+		}
+
+		if (Attributes.HasTangentMode())
+		{
+			// Mirror tangent weight
+			if (Attributes.HasArriveTangentWeight() && Attributes.HasLeaveTangentWeight())
+			{
+				float ArriveWeightTemp = Attributes.GetArriveTangentWeight();
+				float LeaveWeightTemp = Attributes.GetLeaveTangentWeight();
+				Attributes.SetArriveTangentWeight(LeaveWeightTemp);
+				Attributes.SetLeaveTangentWeight(ArriveWeightTemp);
+			}
+			ERichCurveTangentMode TangentMode = Attributes.GetTangentMode();
+			Attributes.SetTangentMode(TangentMode);
+		}
+	}
+}
+
+void FCurveEditor::FlipCurveVertical(TArray<FKeyPosition>& AllKeyPositions, TArray<FKeyAttributes>& AllKeyAttributes, ECurveFlipRangeType RangeType,
+											float InRangeMin, float InRangeMax, double CurveMinVal, double CurveMaxVal)
+{
+	float RangeMin = FLT_MAX;
+	float RangeMax = -FLT_MAX;
+
+	if (RangeType == ECurveFlipRangeType::CurveRange)
+	{
+		RangeMin = FMath::Min(CurveMinVal, RangeMin);
+		RangeMax = FMath::Max(CurveMaxVal, RangeMax);
+	}
+	else if (RangeType == ECurveFlipRangeType::KeyRange)
+	{
+		for (int32 Index = AllKeyPositions.Num() - 1; Index >= 0; --Index)
+		{
+			FKeyPosition& Position = AllKeyPositions[Index];
+
+			RangeMin = FMath::Min(Position.OutputValue, RangeMin);
+			RangeMax = FMath::Max(Position.OutputValue, RangeMax);
+		}
+	}
+	else
+	{
+		RangeMin = InRangeMin;
+		RangeMax = InRangeMax;
+	}
+
+	// Loop through all keys to adjust positions and tangents
+	for (int32 Index = AllKeyPositions.Num() - 1; Index >= 0; --Index)
+	{
+		FKeyPosition& Position = AllKeyPositions[Index];
+		FKeyAttributes& Attributes = AllKeyAttributes[Index];
+
+		// Mirror y value
+		Position.OutputValue = RangeMax - Position.OutputValue + RangeMin;
+
+		// Mirror tangent
+		if (Attributes.HasArriveTangent())
+		{
+			float ArriveTemp = Attributes.GetArriveTangent();
+			Attributes.SetArriveTangent(-ArriveTemp);
+		}
+		if (Attributes.HasLeaveTangent())
+		{
+			float LeaveTemp = Attributes.GetLeaveTangent();
+			Attributes.SetLeaveTangent(-LeaveTemp);
+		}
+		if (Attributes.HasTangentMode())
+		{
+			ERichCurveTangentMode TangentMode = Attributes.GetTangentMode();
+			Attributes.SetTangentMode(TangentMode);
+		}
+	}
+}
+
+void FCurveEditor::FlipCurve(ECurveFlipDirection Direction)
+{
+	FScopedTransaction Transaction(LOCTEXT("FlipCurve", "Flip Curve"));
+	for (FCurveModelID ID : GetEditedCurves())
+	{
+		if (FCurveModel* Curve = FindCurve(ID))
+		{
+			// Init key handles
+			TArray<FKeyHandle> KeyHandles;
+			Curve->GetKeys(TNumericLimits<double>::Lowest(), TNumericLimits<double>::Max(), TNumericLimits<double>::Lowest(), TNumericLimits<double>::Max(), KeyHandles);
+
+			// Init key positions
+			TArray<FKeyPosition> AllKeyPositions;
+			AllKeyPositions.SetNum(KeyHandles.Num());
+			Curve->GetKeyPositions(KeyHandles, AllKeyPositions);
+
+			// Init key attributes
+			TArray<FKeyAttributes> AllKeyAttributes;
+			AllKeyAttributes.SetNum(KeyHandles.Num());
+			Curve->GetKeyAttributes(KeyHandles, AllKeyAttributes);
+
+			// If flipping horizontally
+			if (Direction == ECurveFlipDirection::Horizontal)
+			{
+				double MinTime, MaxTime;
+				Curve->GetTimeRange(MinTime, MaxTime);
+				FlipCurveHorizontal(AllKeyPositions, AllKeyAttributes, HorizontalCurveFlipRangeSettings.RangeType,
+										   HorizontalCurveFlipRangeSettings.MinRange, HorizontalCurveFlipRangeSettings.MaxRange, MinTime, MaxTime);
+			}
+
+			// If flipping vertically
+			if (Direction == ECurveFlipDirection::Vertical)
+			{
+				double MinVal, MaxVal;
+				Curve->GetValueRange(MinVal, MaxVal);
+				FlipCurveVertical(AllKeyPositions, AllKeyAttributes, VerticalCurveFlipRangeSettings.RangeType,
+										 VerticalCurveFlipRangeSettings.MinRange, VerticalCurveFlipRangeSettings.MaxRange, MinVal, MaxVal);
+			}
+
+			if (AllKeyPositions.Num() > 0)
+			{
+				Curve->Modify();
+				Curve->SetKeyPositions(KeyHandles, AllKeyPositions);
+				Curve->SetKeyAttributes(KeyHandles, AllKeyAttributes);
+			}
+		}
+	}
+}
 
 void
 FCurveEditor::ToggleExpandCollapseNodes(bool bRecursive)
diff --git a/Engine/Source/Editor/CurveEditor/Private/CurveEditorCommands.cpp b/Engine/Source/Editor/CurveEditor/Private/CurveEditorCommands.cpp
--- a/Engine/Source/Editor/CurveEditor/Private/CurveEditorCommands.cpp
+++ b/Engine/Source/Editor/CurveEditor/Private/CurveEditorCommands.cpp
@@ -19,6 +19,9 @@ void FCurveEditorCommands::RegisterCommands()
 	UI_COMMAND(ToggleInputSnapping, "Input Snapping", "Toggle Time Snapping", EUserInterfaceActionType::ToggleButton, FInputChord());
 	UI_COMMAND(ToggleOutputSnapping, "Output Snapping", "Toggle Value Snapping", EUserInterfaceActionType::ToggleButton, FInputChord());
 
+	UI_COMMAND(FlipCurveHorizontal, "Flip Curve Horizontal", "Flip curve horizontally", EUserInterfaceActionType::Button, FInputChord());
+	UI_COMMAND(FlipCurveVertical, "Flip Curve Vertical", "Flip curve vertically", EUserInterfaceActionType::Button, FInputChord());
+
 	UI_COMMAND(ToggleExpandCollapseNodes, "Expand/Collapse Nodes", "Toggle expand or collapse selected nodes", EUserInterfaceActionType::Button, FInputChord(EKeys::V) );
 	UI_COMMAND(ToggleExpandCollapseNodesAndDescendants, "Expand/Collapse Nodes and Descendants", "Toggle expand or collapse selected nodes and descendants", EUserInterfaceActionType::Button, FInputChord(EModifierKey::Shift, EKeys::V) );
 
diff --git a/Engine/Source/Editor/CurveEditor/Private/SCurveEditorPanel.cpp b/Engine/Source/Editor/CurveEditor/Private/SCurveEditorPanel.cpp
--- a/Engine/Source/Editor/CurveEditor/Private/SCurveEditorPanel.cpp
+++ b/Engine/Source/Editor/CurveEditor/Private/SCurveEditorPanel.cpp
@@ -1457,6 +1457,40 @@ TSharedPtr<FExtender> SCurveEditorPanel::GetToolbarExtender()
 			}
 			ToolBarBuilder.EndSection();
 
+			ToolBarBuilder.BeginSection("Flip Curve");
+			{
+				ToolBarBuilder.AddToolBarButton(
+					FCurveEditorCommands::Get().FlipCurveHorizontal,
+					NAME_None,
+					TAttribute<FText>(),
+					TAttribute<FText>(),
+					FSlateIcon(FAppStyle::Get().GetStyleSetName(), "GenericCurveEditor.FlipCurveHorizontal"));
+
+				ToolBarBuilder.AddComboButton(
+					FUIAction(),
+					FOnGetContent::CreateSP(InEditorPanel, &SCurveEditorPanel::MakeFlipCurveMenu, FCurveEditor::ECurveFlipDirection::Horizontal),
+					LOCTEXT("FlipCurveHorizontalOptions", "Flip Curve Horizontal"),
+					LOCTEXT("FlipCurveHorizontalOptionsToolTip", "Choose the range for horizontal curve flip"),
+					TAttribute<FSlateIcon>(),
+					true);
+
+				ToolBarBuilder.AddToolBarButton(
+					FCurveEditorCommands::Get().FlipCurveVertical,
+					NAME_None,
+					TAttribute<FText>(),
+					TAttribute<FText>(),
+					FSlateIcon(FAppStyle::Get().GetStyleSetName(), "GenericCurveEditor.FlipCurveVertical"));
+
+				ToolBarBuilder.AddComboButton(
+					FUIAction(),
+					FOnGetContent::CreateSP(InEditorPanel, &SCurveEditorPanel::MakeFlipCurveMenu, FCurveEditor::ECurveFlipDirection::Vertical),
+					LOCTEXT("FlipCurveVerticalOptions", "Flip Curve Vertical"),
+					LOCTEXT("FlipCurveVerticalOptionsToolTip", "Choose the range for vertical curve flip"),
+					TAttribute<FSlateIcon>(),
+					true);
+			}
+			ToolBarBuilder.EndSection();
+
 			ToolBarBuilder.BeginSection("Tangents");
 			{
 				ToolBarBuilder.BeginStyleOverride("CurveEditorToolBar.TangentMode"); // Min size for combo button so size stays consistent independent of label size
@@ -1574,6 +1608,140 @@ TSharedRef<SWidget> SCurveEditorPanel::MakeGridSpacingMenu()
 	return OutputSnapWidget;
 }
 
+TSharedRef<SWidget> SCurveEditorPanel::MakeFlipCurveMenu(FCurveEditor::ECurveFlipDirection Direction)
+{
+	FMenuBuilder MenuBuilder(true, nullptr);
+	MenuBuilder.BeginSection("Header", LOCTEXT("CurveEditorMenuCurveFlipgHeader", "Curve Flip Settings"));
+
+	FCurveEditor::FCurveFlipRangeSettings* RangeSetting = (Direction == FCurveEditor::ECurveFlipDirection::Horizontal) ? &CurveEditor->HorizontalCurveFlipRangeSettings : &CurveEditor->VerticalCurveFlipRangeSettings;
+	
+	// For the x axis, the range of keys and the range of curve should be the same
+	if (Direction == FCurveEditor::ECurveFlipDirection::Horizontal)
+	{
+		// Add menu entry: "Use key/curve range"
+		MenuBuilder.AddMenuEntry(
+			LOCTEXT("UseKeyRange", "Use Key/Curve Range"),
+			LOCTEXT("UseKeyRangeTooltip", "Flip curve within the range of keys/curve"),
+			FSlateIcon(),
+			FUIAction(
+				FExecuteAction::CreateLambda([RangeSetting]()
+					{
+						RangeSetting->RangeType = FCurveEditor::ECurveFlipRangeType::CurveRange;
+					}),
+				FCanExecuteAction(),
+				FIsActionChecked::CreateLambda([RangeSetting]() { return RangeSetting->RangeType == FCurveEditor::ECurveFlipRangeType::CurveRange; })
+			),
+			NAME_None,
+			EUserInterfaceActionType::RadioButton
+		);
+	}
+
+	// For the y axis, the range of curve can be different than range of keys if tangent is set
+	if (Direction == FCurveEditor::ECurveFlipDirection::Vertical)
+	{
+		// Add menu entry: "Use key range"
+		MenuBuilder.AddMenuEntry(
+			LOCTEXT("UseKeyRange", "Use Key Range"),
+			LOCTEXT("UseKeyRangeTooltip", "Flip curve within the range of the keys"),
+			FSlateIcon(),
+			FUIAction(
+				FExecuteAction::CreateLambda([RangeSetting]()
+					{
+						RangeSetting->RangeType = FCurveEditor::ECurveFlipRangeType::KeyRange;
+					}),
+				FCanExecuteAction(),
+				FIsActionChecked::CreateLambda([RangeSetting]() { return RangeSetting->RangeType == FCurveEditor::ECurveFlipRangeType::KeyRange; })
+			),
+			NAME_None,
+			EUserInterfaceActionType::RadioButton
+		);
+
+		// Add menu entry: "Use curve range"
+		MenuBuilder.AddMenuEntry(
+			LOCTEXT("UseCurveRange", "Use Curve Range"),
+			LOCTEXT("UseCurveRangeTooltip", "Flip curve within the range of the curve"),
+			FSlateIcon(),
+			FUIAction(
+				FExecuteAction::CreateLambda([RangeSetting]()
+					{
+						RangeSetting->RangeType = FCurveEditor::ECurveFlipRangeType::CurveRange;
+					}),
+				FCanExecuteAction(),
+				FIsActionChecked::CreateLambda([RangeSetting]() { return RangeSetting->RangeType == FCurveEditor::ECurveFlipRangeType::CurveRange; })
+			),
+			NAME_None,
+			EUserInterfaceActionType::RadioButton
+		);
+	}
+
+	// Add "Custom Range" menu entry
+	MenuBuilder.AddMenuEntry(
+		LOCTEXT("CustomRange", "Use Custom Range"),
+		LOCTEXT("CustomRangeTooltip", "Flip curve within the range of float inputs (Default range is 0 to 1)"),
+		FSlateIcon(),
+		FUIAction(
+			FExecuteAction::CreateLambda([RangeSetting]()
+				{
+					RangeSetting->RangeType = FCurveEditor::ECurveFlipRangeType::CustomRange;
+				}),
+			FCanExecuteAction(),
+			FIsActionChecked::CreateLambda([RangeSetting]() { return RangeSetting->RangeType == FCurveEditor::ECurveFlipRangeType::CustomRange; })
+		),
+		NAME_None,
+		EUserInterfaceActionType::RadioButton
+	);
+
+	// Create a custom widget that includes two input boxes: one for "Min" and one for "Max"
+	TSharedRef<SWidget> CustomRangeWidget =
+		SNew(SVerticalBox)
+		+ SVerticalBox::Slot()
+		.AutoHeight()
+		.Padding(FMargin(0, 6, 0, 0))
+		[
+			SNew(STextBlock)
+				.Text(LOCTEXT("MinRangeLabel", "Min:"))
+		]
+		+ SVerticalBox::Slot()
+		.AutoHeight()
+		[
+			SNew(SEditableTextBox)
+				.MinDesiredWidth(60.0f)
+				.Text_Lambda([RangeSetting]() -> FText {
+				return FText::AsNumber(RangeSetting->MinRange);
+					})
+				.OnTextCommitted_Lambda([RangeSetting](const FText& InText, ETextCommit::Type CommitType)
+					{
+						RangeSetting->MinRange = FCString::Atof(*InText.ToString());
+					})
+		]
+		+ SVerticalBox::Slot()
+		.AutoHeight()
+		.Padding(FMargin(0, 8, 0, 0))
+		[
+			SNew(STextBlock)
+				.Text(LOCTEXT("MaxRangeLabel", "Max:"))
+		]
+		+ SVerticalBox::Slot()
+		.AutoHeight()
+		[
+			SNew(SEditableTextBox)
+				.MinDesiredWidth(60.0f)
+				.Text_Lambda([RangeSetting]() -> FText {
+				return FText::AsNumber(RangeSetting->MaxRange);
+					})
+				.OnTextCommitted_Lambda([RangeSetting](const FText& InText, ETextCommit::Type CommitType)
+					{
+						RangeSetting->MaxRange = FCString::Atof(*InText.ToString());
+					})
+		];
+
+	// Only enable the custom range input widget when "Custom Range" is selected.
+	CustomRangeWidget->SetEnabled(TAttribute<bool>::CreateLambda([RangeSetting]() { return RangeSetting->RangeType == FCurveEditor::ECurveFlipRangeType::CustomRange; }));
+
+	MenuBuilder.AddWidget(CustomRangeWidget, LOCTEXT("CustomRangeInputs", "Range"));
+	return MenuBuilder.MakeWidget();
+}
+
 TSharedRef<SWidget> SCurveEditorPanel::MakeAxisSnapMenu()
 {
 	FMenuBuilder MenuBuilder(true, GetCommands());
diff --git a/Engine/Source/Editor/CurveEditor/Public/CurveEditor.h b/Engine/Source/Editor/CurveEditor/Public/CurveEditor.h
--- a/Engine/Source/Editor/CurveEditor/Public/CurveEditor.h
+++ b/Engine/Source/Editor/CurveEditor/Public/CurveEditor.h
@@ -499,6 +499,47 @@ public:
 	bool IsOutputSnappingEnabled() const;
 	void ToggleOutputSnapping();
 
+	/** Curve flip direction types **/
+	enum class ECurveFlipDirection : uint8
+	{
+		Horizontal,
+		Vertical
+	};
+
+	/** Curve Flip selection types **/
+	enum class ECurveFlipRangeType : uint8
+	{
+		KeyRange,
+		CurveRange,
+		CustomRange
+	};
+
+	/** Settings for Curve Flip dropdown menu selections **/
+	struct FCurveFlipRangeSettings
+	{
+		ECurveFlipRangeType RangeType;
+		float MinRange;
+		float MaxRange;
+
+		FCurveFlipRangeSettings()
+			: RangeType(ECurveFlipRangeType::CustomRange)
+			, MinRange(0.0f)
+			, MaxRange(1.0f)
+		{
+		}
+	};
+	FCurveFlipRangeSettings HorizontalCurveFlipRangeSettings;
+	FCurveFlipRangeSettings VerticalCurveFlipRangeSettings;
+
+	/**
+	 * Flip the selected curve horizontally or vertically
+	 */
+	void FlipCurve(ECurveFlipDirection Direction);
+	static void FlipCurveHorizontal(TArray<FKeyPosition>& AllKeyPositions, TArray<FKeyAttributes>& AllKeyAttributes, ECurveFlipRangeType RangeType,
+											float InRangeMin, float InRangeMax, double CurveMinTime, double CurveMaxTime);
+	static void FlipCurveVertical(TArray<FKeyPosition>& AllKeyPositions, TArray<FKeyAttributes>& AllKeyAttributes, ECurveFlipRangeType RangeType,
+											float InRangeMin, float InRangeMax, double CurveMinVal, double CurveMaxVal);
+
 public:
 
 	/**
diff --git a/Engine/Source/Editor/CurveEditor/Public/CurveEditorCommands.h b/Engine/Source/Editor/CurveEditor/Public/CurveEditorCommands.h
--- a/Engine/Source/Editor/CurveEditor/Public/CurveEditorCommands.h
+++ b/Engine/Source/Editor/CurveEditor/Public/CurveEditorCommands.h
@@ -36,6 +36,9 @@ public:
 	TSharedPtr<FUICommandInfo> ToggleInputSnapping;
 	TSharedPtr<FUICommandInfo> ToggleOutputSnapping;
 
+	TSharedPtr<FUICommandInfo> FlipCurveHorizontal;
+	TSharedPtr<FUICommandInfo> FlipCurveVertical;
+
 	TSharedPtr<FUICommandInfo> ToggleExpandCollapseNodes;
 	TSharedPtr<FUICommandInfo> ToggleExpandCollapseNodesAndDescendants;
 
diff --git a/Engine/Source/Editor/CurveEditor/Public/SCurveEditorPanel.h b/Engine/Source/Editor/CurveEditor/Public/SCurveEditorPanel.h
--- a/Engine/Source/Editor/CurveEditor/Public/SCurveEditorPanel.h
+++ b/Engine/Source/Editor/CurveEditor/Public/SCurveEditorPanel.h
@@ -181,7 +181,7 @@ class CURVEEDITOR_API SCurveEditorPanel : public SCompoundWidget
 
 	/** Enable/disable pending focus */
 	void EnablePendingFocusOnHovering(const bool InEnabled);
-	
+
 private:
 	// SWidget Interface
 	virtual void Tick(const FGeometry& AllottedGeometry, const double InCurrentTime, const float InDeltaTime) override;
@@ -199,6 +199,9 @@ private:
 	TSharedRef<SWidget> MakeGridSpacingMenu();
 	TSharedRef<SWidget> MakeAxisSnapMenu();
 
+	/* Curve Flip make menu function and settings */
+	TSharedRef<SWidget> MakeFlipCurveMenu(FCurveEditor::ECurveFlipDirection Direction);
+
 	bool IsInlineEditPanelEditable() const;
 	EVisibility ShouldInstructionOverlayBeVisible() const;
 
diff --git a/Engine/Source/Editor/EditorStyle/Private/StarshipStyle.cpp b/Engine/Source/Editor/EditorStyle/Private/StarshipStyle.cpp
--- a/Engine/Source/Editor/EditorStyle/Private/StarshipStyle.cpp
+++ b/Engine/Source/Editor/EditorStyle/Private/StarshipStyle.cpp
@@ -7670,6 +7670,10 @@ void FStarshipEditorStyle::FStyle::SetupToolkitStyles()
 		Set("GenericCurveEditor.ToggleInputSnapping", new IMAGE_BRUSH_SVG("Starship/GenericCurveEditor/AdjustKeySnapFrameHorizontal", Icon20x20));
 		Set("GenericCurveEditor.ToggleOutputSnapping", new IMAGE_BRUSH_SVG("Starship/GenericCurveEditor/AdjustKeySnapFrameVertical", Icon20x20));
 
+		// Flip Curve
+		Set("GenericCurveEditor.FlipCurveHorizontal", new IMAGE_BRUSH_SVG("Starship/GenericCurveEditor/FlipHorizontal", Icon20x20));
+		Set("GenericCurveEditor.FlipCurveVertical", new IMAGE_BRUSH_SVG("Starship/GenericCurveEditor/FlipVertical", Icon20x20));
+
 		// Tangent Types
 		Set("GenericCurveEditor.InterpolationCubicSmartAuto", new IMAGE_BRUSH_SVG("Starship/GenericCurveEditor/TangentsCubicSmartAuto", Icon20x20));
 		Set("GenericCurveEditor.InterpolationCubicAuto", new IMAGE_BRUSH_SVG("Starship/GenericCurveEditor/TangentsCubicAuto", Icon20x20));
