Fix: duplicate LOCTEXT keys

diff --git a/Engine/Source/Editor/CurveEditor/Private/SCurveEditorPanel.cpp b/Engine/Source/Editor/CurveEditor/Private/SCurveEditorPanel.cpp
--- a/Engine/Source/Editor/CurveEditor/Private/SCurveEditorPanel.cpp
+++ b/Engine/Source/Editor/CurveEditor/Private/SCurveEditorPanel.cpp
@@ -1620,8 +1620,8 @@ TSharedRef<SWidget> SCurveEditorPanel::MakeFlipCurveMenu(FCurveEditor::ECurveFli
 	{
 		// Add menu entry: "Use key/curve range"
 		MenuBuilder.AddMenuEntry(
-			LOCTEXT("UseKeyRange", "Use Key/Curve Range"),
-			LOCTEXT("UseKeyRangeTooltip", "Flip curve within the range of keys/curve"),
+			LOCTEXT("UseKeyCurveRange", "Use Key/Curve Range"),
+			LOCTEXT("UseKeyCurveRangeTooltip", "Flip curve within the range of keys/curve"),
 			FSlateIcon(),
 			FUIAction(
 				FExecuteAction::CreateLambda([RangeSetting]()
