Curve Editor: fixed curves not rendering in the animation sequence editor
- 40614307 follow up

#rnx

diff --git a/Engine/Source/Editor/CurveEditor/Private/SCurveViewerPanel.cpp b/Engine/Source/Editor/CurveEditor/Private/SCurveViewerPanel.cpp
--- a/Engine/Source/Editor/CurveEditor/Private/SCurveViewerPanel.cpp
+++ b/Engine/Source/Editor/CurveEditor/Private/SCurveViewerPanel.cpp
@@ -70,13 +70,13 @@ void SCurveViewerPanel::DrawCurves(const FGeometry& AllottedGeometry, FSlateWind
 
 	float DashOffset = static_cast<float>(GetViewSpace().PixelsPerInput() * GetViewSpace().GetInputMin());
 
-	for (const FCurveDrawParams& Params : CachedDrawParams)
+	for (const FCurveDrawParams& Params : CurveDrawParamsCache->GetCurveDrawParams())
 	{
 		if (Params.DashLengthPx > 0.f)
 		{
 			TArray<FVector2f> NewVector;
 			NewVector.Reserve(Params.InterpolatingPoints.Num());
-			for (FVector2d Vect : Params.InterpolatingPoints)
+			for (const FVector2d& Vect : Params.InterpolatingPoints)
 			{
 				NewVector.Add(UE::Slate::CastToVector2f(Vect));
 			}
