Curve Editor: Refine paste behaviour:

- Added paste relative (alt+v for): aligns the first key with the key closest to the current scruberr position
- Fixed minor bug from 37658913: Effectively, PasteOverwriteRange was set up to paste additively. The ctrl+v continues to be OverwriteRange but
- removed ctrl+shift+v from PasteOverwriteRange, and
- reassigned ctrl+shift+v to PasteAndMerge.
- PasteOverwriteRange is now only used in context menu
- Context menu contains PasteOverwriteRange, PasteAndMerge, PasteRelative

Effective command bindings:
- Ctrl+V: Paste Overwrite
- Ctrl + Shift + V: Paste & Merge
- Alt + V: Paste relative

diff --git a/Engine/Source/Editor/CurveEditor/Private/CurveEditor.cpp b/Engine/Source/Editor/CurveEditor/Private/CurveEditor.cpp
--- a/Engine/Source/Editor/CurveEditor/Private/CurveEditor.cpp
+++ b/Engine/Source/Editor/CurveEditor/Private/CurveEditor.cpp
@@ -38,6 +38,7 @@
 #include "Math/UnrealMathUtility.h"
 #include "Misc/FrameNumber.h"
 #include "Misc/FrameTime.h"
+#include "Misc/KeyPasteArgs.h"
 #include "Modules/ModuleManager.h"
 #include "SCurveEditor.h" // for access to LogCurveEditor
 #include "SCurveEditorPanel.h"
@@ -442,6 +443,7 @@ bool FCurveEditor::ShouldAutoFrame() const
 
 void FCurveEditor::BindCommands()
 {
+	using namespace UE::CurveEditor;
 	UCurveEditorSettings* CurveSettings = Settings;
 		
 	CommandList->MapAction(FGenericCommands::Get().Undo,   FExecuteAction::CreateLambda([]{ GEditor->UndoTransaction(); }));
@@ -450,8 +452,15 @@ void FCurveEditor::BindCommands()
 
 	CommandList->MapAction(FGenericCommands::Get().Cut, FExecuteAction::CreateSP(this, &FCurveEditor::CutSelection));
 	CommandList->MapAction(FGenericCommands::Get().Copy, FExecuteAction::CreateSP(this, &FCurveEditor::CopySelection));
-	CommandList->MapAction(FGenericCommands::Get().Paste, FExecuteAction::CreateSP(this, &FCurveEditor::PasteKeys, TSet<FCurveModelID>(), true));
-	CommandList->MapAction(FCurveEditorCommands::Get().PasteOverwriteRange, FExecuteAction::CreateSP(this, &FCurveEditor::PasteKeys, TSet<FCurveModelID>(), false));
+	CommandList->MapAction(FGenericCommands::Get().Paste, FExecuteAction::CreateSP(this, &FCurveEditor::PasteKeys,
+		FKeyPasteArgs{ .Mode = ECurveEditorPasteMode::OverwriteRange,	.Flags = ECurveEditorPasteFlags::Default }
+		));
+	CommandList->MapAction(FCurveEditorCommands::Get().PasteAndMerge, FExecuteAction::CreateSP(this, &FCurveEditor::PasteKeys,
+		FKeyPasteArgs{ .Mode = ECurveEditorPasteMode::Merge,			.Flags = ECurveEditorPasteFlags::Default }
+		));
+	CommandList->MapAction(FCurveEditorCommands::Get().PasteRelative, FExecuteAction::CreateSP(this, &FCurveEditor::PasteKeys,
+		FKeyPasteArgs{ .Mode = ECurveEditorPasteMode::OverwriteRange,	.Flags = ECurveEditorPasteFlags::Default | ECurveEditorPasteFlags::Relative }
+		));
 
 	CommandList->MapAction(FCurveEditorCommands::Get().ZoomToFit, FExecuteAction::CreateSP(this, &FCurveEditor::ZoomToFit, EAxisList::All));
 	CommandList->MapAction(FCurveEditorCommands::Get().ZoomToFitHorizontal, FExecuteAction::CreateSP(this, &FCurveEditor::ZoomToFit, EAxisList::X));
@@ -1818,51 +1827,106 @@ TSet<FCurveModelID> FCurveEditor::GetTargetCurvesForPaste() const
 
 bool FCurveEditor::CopyBufferCurveToCurveID(const UCurveEditorCopyableCurveKeys* InSourceCurve, const FCurveModelID InTargetCurve, TOptional<double> InTimeOffset, const bool bInAddToSelection, const bool bInOverwriteRange)
 {
-	FCurveModel* TargetCurve = FindCurve(InTargetCurve);
-	if (!InSourceCurve || !TargetCurve)
-	{
-		return false;
-	}
+	using namespace UE::CurveEditor;
+	return CopyBufferCurveToCurveID(InSourceCurve, InTargetCurve, InTimeOffset,
+		bInOverwriteRange ? ECurveEditorPasteMode::OverwriteRange : ECurveEditorPasteMode::Merge,
+		bInAddToSelection ? ECurveEditorPasteFlags::Default | ECurveEditorPasteFlags::SetSelection : ECurveEditorPasteFlags::Default
+		);
+}
 
-	// Sometimes when you paste you want to delete any keys that already exist in the timerange you'll be replacing
-	// because mixing the pasted results with the original results wouldn't make any sense.
-	if (bInOverwriteRange)
+namespace UE::CurveEditor::PasteDetail
+{
+static void RemovePastedKeysInRange(
+	const UCurveEditorCopyableCurveKeys* InSourceCurve, FCurveModel* InTargetCurveModel, const TOptional<double>& InTimeOffset,
+	double InCurrentTime
+	)
+{
+	TArray<FKeyHandle> KeysToRemove;
+	double MinKeyTime = TNumericLimits<double>::Max();
+	double MaxKeyTime = TNumericLimits<double>::Lowest(); 
+	for (int32 Index = 0; Index < InSourceCurve->KeyPositions.Num(); ++Index)
 	{
-		double CurrentTime = 0.0;
-		if (TSharedPtr<ITimeSliderController> TimeSliderController = WeakTimeSliderController.Pin())
+		FKeyPosition KeyPosition = InSourceCurve->KeyPositions[Index];
+		if (InTimeOffset.IsSet())
 		{
-			FFrameRate TickResolution = TimeSliderController->GetTickResolution();
-			CurrentTime = TickResolution.AsSeconds(TimeSliderController->GetScrubPosition());
+			KeyPosition.InputValue += InTimeOffset.GetValue();
 		}
-		TArray<FKeyHandle> KeysToRemove;
-		double MinKeyTime = TNumericLimits<double>::Max();
-		double MaxKeyTime = TNumericLimits<double>::Lowest(); 
-		for (int32 Index = 0; Index < InSourceCurve->KeyPositions.Num(); ++Index)
+		if (KeyPosition.InputValue < MinKeyTime)
 		{
-			FKeyPosition KeyPosition = InSourceCurve->KeyPositions[Index];
-			if (InTimeOffset.IsSet())
-			{
-				KeyPosition.InputValue += InTimeOffset.GetValue();
-			}
-			if (KeyPosition.InputValue < MinKeyTime)
-			{
-				MinKeyTime = KeyPosition.InputValue;
-			}
-			if (KeyPosition.InputValue > MaxKeyTime)
-			{
-				MaxKeyTime = KeyPosition.InputValue;
-			}
+			MinKeyTime = KeyPosition.InputValue;
 		}
-
-		// Just double checking we actually set a Min/Max time so we don't wipe out every key to infinity.
-		if (InSourceCurve->KeyPositions.Num() > 0)
+		if (KeyPosition.InputValue > MaxKeyTime)
 		{
-			TargetCurve->GetKeys(MinKeyTime, MaxKeyTime, TNumericLimits<double>::Lowest(), TNumericLimits<double>::Max(), KeysToRemove);
+			MaxKeyTime = KeyPosition.InputValue;
 		}
+	}
+
+	// Just double checking we actually set a Min/Max time so we don't wipe out every key to infinity.
+	if (InSourceCurve->KeyPositions.Num() > 0)
+	{
+		InTargetCurveModel->GetKeys(MinKeyTime, MaxKeyTime, TNumericLimits<double>::Lowest(), TNumericLimits<double>::Max(), KeysToRemove);
+	}
+
+	InTargetCurveModel->RemoveKeys(KeysToRemove, InCurrentTime);
+}
+/** @return Height to add to each pasted key to make the range relative to the closest key to the left of the pasted range. */
+static double FindRelativeKeyPasteInset(const UCurveEditorCopyableCurveKeys* InSourceCurve, FCurveModel* InTargetCurveModel, double InCurrentTime)
+{
+	const FKeyPosition* MinElement = Algo::MinElementBy(InSourceCurve->KeyPositions, [](const FKeyPosition& Position)
+	{
+		return Position.InputValue;
+	});
+	if (!MinElement) 
+	{
+		// Nothing to paste
+		return 0.0;
+	}
+
+	TOptional<FKeyHandle> ClosestPrevious, ClosestNext;
+	InTargetCurveModel->GetClosestKeysTo(InCurrentTime, ClosestPrevious, ClosestNext);
+	if (!ClosestPrevious)
+	{
+		// No previous key -> no inset to apply
+		return 0.0;
+	}
+
+	FKeyPosition Position{};
+	InTargetCurveModel->GetKeyPositions(TConstArrayView<FKeyHandle>(&ClosestPrevious.GetValue(), 1), TArrayView<FKeyPosition>(&Position, 1));
+	return Position.OutputValue - MinElement->OutputValue; // Bring all values down to the closest value to the left.
+}
+}
+
+bool FCurveEditor::CopyBufferCurveToCurveID(
+	const UCurveEditorCopyableCurveKeys* InSourceCurve, const FCurveModelID InTargetCurve,
+	TOptional<double> InTimeOffset, UE::CurveEditor::ECurveEditorPasteMode InMode, UE::CurveEditor::ECurveEditorPasteFlags InFlags
+	)
+{
+	using namespace UE::CurveEditor;
+	
+	FCurveModel* TargetCurveModel = FindCurve(InTargetCurve);
+	if (!InSourceCurve || !TargetCurveModel)
+	{
+		return false;
+	}
+	
+	double CurrentTime = 0.0;
+	if (TSharedPtr<ITimeSliderController> TimeSliderController = WeakTimeSliderController.Pin())
+	{
+		FFrameRate TickResolution = TimeSliderController->GetTickResolution();
+		CurrentTime = TickResolution.AsSeconds(TimeSliderController->GetScrubPosition());
+	}
 
-		TargetCurve->RemoveKeys(KeysToRemove, CurrentTime);
+	// Sometimes when you paste you want to delete any keys that already exist in the timerange you'll be replacing
+	// because mixing the pasted results with the original results wouldn't make any sense.
+	if (InMode == ECurveEditorPasteMode::OverwriteRange)
+	{
+		PasteDetail::RemovePastedKeysInRange(InSourceCurve, TargetCurveModel, InTimeOffset, CurrentTime);
 	}
 
+	// If pasting relative, bring all pasted values down to the first key to the left of the scrubber.
+	const double ValueInset = EnumHasAnyFlags(InFlags, ECurveEditorPasteFlags::Relative)
+		? PasteDetail::FindRelativeKeyPasteInset(InSourceCurve, TargetCurveModel, CurrentTime) : 0.0;
+
 	for (int32 Index = 0; Index < InSourceCurve->KeyPositions.Num(); ++Index)
 	{
 		FKeyPosition KeyPosition = InSourceCurve->KeyPositions[Index];
@@ -1871,8 +1935,10 @@ bool FCurveEditor::CopyBufferCurveToCurveID(const UCurveEditorCopyableCurveKeys*
 			KeyPosition.InputValue += InTimeOffset.GetValue();
 		}
 
-		TOptional<FKeyHandle> KeyHandle = TargetCurve->AddKey(KeyPosition, InSourceCurve->KeyAttributes[Index]);
-		if (KeyHandle.IsSet() && bInAddToSelection)
+		KeyPosition.OutputValue += ValueInset;
+
+		TOptional<FKeyHandle> KeyHandle = TargetCurveModel->AddKey(KeyPosition, InSourceCurve->KeyAttributes[Index]);
+		if (KeyHandle.IsSet() && EnumHasAnyFlags(InFlags, ECurveEditorPasteFlags::SetSelection))
 		{
 			Selection.Add(FCurvePointHandle(InTargetCurve, ECurvePointType::Key, KeyHandle.GetValue()));
 		}
@@ -1882,6 +1948,16 @@ bool FCurveEditor::CopyBufferCurveToCurveID(const UCurveEditorCopyableCurveKeys*
 }
 
 void FCurveEditor::PasteKeys(TSet<FCurveModelID> CurveModelIDs, const bool bInOverwriteRange)
+{
+	using namespace UE::CurveEditor;
+
+	FKeyPasteArgs Args;
+	Args.CurveModelIds = CurveModelIDs;
+	Args.Mode = bInOverwriteRange ? ECurveEditorPasteMode::OverwriteRange : ECurveEditorPasteMode::Merge;
+	PasteKeys(Args);
+}
+
+void FCurveEditor::PasteKeys(UE::CurveEditor::FKeyPasteArgs InArgs)
 {
 	// Grab the text to paste from the clipboard
 	FString TextToImport;
@@ -1920,7 +1996,7 @@ void FCurveEditor::PasteKeys(TSet<FCurveModelID> CurveModelIDs, const bool bInOv
 
 	// Figure out which CurveModelIDs we're trying to paste to. If they're not already specified, we try to find hovered curves,
 	// and failing that we try to find all curves.
-	TSet<FCurveModelID> TargetCurves = CurveModelIDs.Num() > 0 ? CurveModelIDs : GetTargetCurvesForPaste();
+	TSet<FCurveModelID> TargetCurves = InArgs.CurveModelIds.Num() > 0 ? InArgs.CurveModelIds : GetTargetCurvesForPaste();
 	
 	if (TargetCurves.Num() == 0)
 	{
@@ -2065,8 +2141,7 @@ void FCurveEditor::PasteKeys(TSet<FCurveModelID> CurveModelIDs, const bool bInOv
 	// Now that we've calculated the source curve for each destination curve, copy them over.
 	for (const TPair<UCurveEditorCopyableCurveKeys*, FCurveModelID>& Pair : CopyPairs)
 	{
-		const bool bAddToSelection = true;
-		CopyBufferCurveToCurveID(Pair.Key, Pair.Value, TimeOffset, bAddToSelection, bInOverwriteRange);
+		CopyBufferCurveToCurveID(Pair.Key, Pair.Value, TimeOffset, InArgs.Mode, InArgs.Flags);
 	}
 
 	if (ShouldAutoFrame())
diff --git a/Engine/Source/Editor/CurveEditor/Private/CurveEditorCommands.cpp b/Engine/Source/Editor/CurveEditor/Private/CurveEditorCommands.cpp
--- a/Engine/Source/Editor/CurveEditor/Private/CurveEditorCommands.cpp
+++ b/Engine/Source/Editor/CurveEditor/Private/CurveEditorCommands.cpp
@@ -71,9 +71,10 @@ void FCurveEditorCommands::RegisterCommands()
 	UI_COMMAND(ToggleShowValueIndicatorLines, "Show Value Indicators", "Show horizontal, dotted lines on the selected min & max valued key handles", EUserInterfaceActionType::ToggleButton, FInputChord());
 
 	UI_COMMAND(AddKeyHovered, "Add Key", "Add a new key to this curve at the current position.", EUserInterfaceActionType::Button, FInputChord(EKeys::MiddleMouseButton) );
-	UI_COMMAND(PasteKeysHovered, "Paste", "Paste clipboard contents", EUserInterfaceActionType::Button, FInputChord(EModifierKey::Control, EKeys::V));
-	UI_COMMAND(PasteOverwriteRange, "Paste Overwrite", "Pastes the keys from the clipboard, overwriting any key in destination track between the first and last pasted keys.", EUserInterfaceActionType::Button, FInputChord(EModifierKey::Control | EModifierKey::Shift, EKeys::V));
-
+	UI_COMMAND(PasteKeysHovered, "Paste", "Pastes the keys from the clipboard overwriting any key in destination track between the first and last pasted keys (default)", EUserInterfaceActionType::Button, FInputChord(EModifierKey::Control, EKeys::V));
+	UI_COMMAND(PasteAndMerge, "Paste & Merge", "Pastes the keys from the clipboard over the keys in the graph.", EUserInterfaceActionType::Button, FInputChord(EModifierKey::Control | EModifierKey::Shift, EKeys::V));
+	UI_COMMAND(PasteRelative, "Paste Relative", "Pastes the keys from the clipboard and aligns them to the nearest key to the left of the scrubber.", EUserInterfaceActionType::Button, FInputChord(EModifierKey::Alt, EKeys::V));
+	
 	UI_COMMAND(AddKeyToAllCurves, "Add Key", "Add a new key to all curves at the current time.", EUserInterfaceActionType::Button, FInputChord(EKeys::Enter) );
 
 	// Curve Editor Colors
diff --git a/Engine/Source/Editor/CurveEditor/Private/CurveEditorContextMenu.cpp b/Engine/Source/Editor/CurveEditor/Private/CurveEditorContextMenu.cpp
--- a/Engine/Source/Editor/CurveEditor/Private/CurveEditorContextMenu.cpp
+++ b/Engine/Source/Editor/CurveEditor/Private/CurveEditorContextMenu.cpp
@@ -144,7 +144,8 @@ void FContextMenuBuilder::BuildKeyMenu(FMenuBuilder& MenuBuilder)
 			// Modify Data
 			MenuBuilder.AddMenuEntry(FGenericCommands::Get().Cut);
 			MenuBuilder.AddMenuEntry(FGenericCommands::Get().Paste);
-			MenuBuilder.AddMenuEntry(FCurveEditorCommands::Get().PasteOverwriteRange);
+			MenuBuilder.AddMenuEntry(FCurveEditorCommands::Get().PasteAndMerge);
+			MenuBuilder.AddMenuEntry(FCurveEditorCommands::Get().PasteRelative);
 			MenuBuilder.AddMenuEntry(FGenericCommands::Get().Copy);
 			MenuBuilder.AddMenuEntry(FGenericCommands::Get().Delete);
 
diff --git a/Engine/Source/Editor/CurveEditor/Private/CurveModel.cpp b/Engine/Source/Editor/CurveEditor/Private/CurveModel.cpp
--- a/Engine/Source/Editor/CurveEditor/Private/CurveModel.cpp
+++ b/Engine/Source/Editor/CurveEditor/Private/CurveModel.cpp
@@ -7,6 +7,49 @@
 #include "Curves/KeyHandle.h"
 #include "HAL/PlatformCrt.h"
 
+void FCurveModel::GetClosestKeysTo(double InTime, TOptional<FKeyHandle>& OutPreviousKeyHandle, TOptional<FKeyHandle>& OutNextKeyHandle) const
+{
+	// Default implementation does linear search. Subclasses can override this with a more efficient implementations.
+	double MinTime, MaxTime;
+	GetTimeRange(MinTime, MaxTime);
+	double MinValue, MaxValue;
+	GetValueRange(MinValue, MaxValue);
+
+	TArray<FKeyHandle> AllHandles;
+	GetKeys(MinTime, MaxTime, MinValue, MaxValue, AllHandles);
+	if (AllHandles.IsEmpty())
+	{
+		return;
+	}
+	
+	TArray<FKeyPosition> AllPositions;
+	AllPositions.SetNumUninitialized(AllHandles.Num());
+	GetKeyPositions(AllHandles, AllPositions);
+
+	// Ideally, we'd use binary search.
+	// However, GetKeyPositions documentation does not dictate that the function returns a sorted array.
+	// For now, this implementation does linear because it's safe.
+	double CurrentClosestLeftTime = MinTime - 1.0; // We'll assume we don't underflow double range
+	double CurrentClosestRightTime = MaxTime + 1.0; // We'll assume we don't overflow double range
+	for (int32 Index = 0; Index < AllPositions.Num(); ++Index)
+	{
+		const FKeyHandle& KeyHandle = AllHandles[Index];
+		const double CurrentTime = AllPositions[Index].InputValue;
+
+		if (CurrentTime <= InTime && CurrentClosestLeftTime < CurrentTime)
+		{
+			CurrentClosestLeftTime = CurrentTime;
+			OutPreviousKeyHandle = KeyHandle;
+		}
+
+		if (CurrentTime >= InTime && CurrentClosestRightTime > CurrentTime)
+		{
+			CurrentClosestRightTime = CurrentTime;
+			OutNextKeyHandle = KeyHandle;
+		}
+	}
+}
+
 void FCurveModel::SetKeyAttributes(TArrayView<const FKeyHandle> InKeys, const FKeyAttributes& InKeyAttributes, EPropertyChangeType::Type ChangeType)
 {
 	TArray<FKeyAttributes> ExpandedAttributes;
diff --git a/Engine/Source/Editor/CurveEditor/Private/Views/SInteractiveCurveEditorView.cpp b/Engine/Source/Editor/CurveEditor/Private/Views/SInteractiveCurveEditorView.cpp
--- a/Engine/Source/Editor/CurveEditor/Private/Views/SInteractiveCurveEditorView.cpp
+++ b/Engine/Source/Editor/CurveEditor/Private/Views/SInteractiveCurveEditorView.cpp
@@ -2203,7 +2203,7 @@ void SInteractiveCurveEditorView::PasteKeys(TSet<FCurveModelID> ToCurves)
 		return;
 	}
 
-	CurveEditor->PasteKeys(ToCurves);
+	CurveEditor->PasteKeys(UE::CurveEditor::FKeyPasteArgs{ ToCurves });
 }
 
 void SInteractiveCurveEditorView::OnCurveEditorToolChanged(FCurveEditorToolID InToolId)
diff --git a/Engine/Source/Editor/CurveEditor/Public/CurveEditor.h b/Engine/Source/Editor/CurveEditor/Public/CurveEditor.h
--- a/Engine/Source/Editor/CurveEditor/Public/CurveEditor.h
+++ b/Engine/Source/Editor/CurveEditor/Public/CurveEditor.h
@@ -30,6 +30,7 @@
 #include "Math/Vector2D.h"
 #include "Misc/AssertionMacros.h"
 #include "Misc/Attribute.h"
+#include "Misc/KeyPasteArgs.h"
 #include "Misc/FrameRate.h"
 #include "Misc/Optional.h"
 #include "Modification//TransactionManager.h"
@@ -560,15 +561,26 @@ public:
 protected:
 	void ImportCopyBufferFromText(const FString& TextToImport, /*out*/ TArray<UCurveEditorCopyBuffer*>& ImportedCopyBuffers) const;
 	TSet<FCurveModelID> GetTargetCurvesForPaste() const;
+	
+	UE_DEPRECATED(5.6, "Use the version that takes in ECurveEditorPasteMode and ECurveEditorPasteFlags instead.")
 	bool CopyBufferCurveToCurveID(const UCurveEditorCopyableCurveKeys* InSourceCurve, const FCurveModelID InTargetCurve, TOptional<double> InTimeOffset, const bool bInAddToSelection, const bool bInOverwriteRange);
-
+	bool CopyBufferCurveToCurveID(const UCurveEditorCopyableCurveKeys* InSourceCurve, const FCurveModelID InTargetCurve, TOptional<double> InTimeOffset,
+		UE::CurveEditor::ECurveEditorPasteMode InMode, UE::CurveEditor::ECurveEditorPasteFlags InFlags
+		);
+	
 	void GetChildCurveModelIDs(const FCurveEditorTreeItemID TreeItemID, TSet<FCurveModelID>& CurveModelIDs) const;
 
 public:
+	/** Paste keys*/
+	UE_DEPRECATED(5.6, "Use the version that takes FKeyPasteArg instead.")
+	void PasteKeys(TSet<FCurveModelID> CurveModelIDs, const bool bInOverwriteRange = false);
+
 	/**
-	 * Paste keys
+	 * Pastes keys.
+	 * @param CurveModelIds Only the curve model IDs specified in this set. If empty, paste all in the clipboard.
+	 * @param Flags Flags that modify the operation.
 	 */
-	void PasteKeys(TSet<FCurveModelID> CurveModelIDs, const bool bInOverwriteRange = false);
+	void PasteKeys(UE::CurveEditor::FKeyPasteArgs InArgs = {});
 
 	/**
 	 * Delete the currently selected keys
diff --git a/Engine/Source/Editor/CurveEditor/Public/CurveEditorCommands.h b/Engine/Source/Editor/CurveEditor/Public/CurveEditorCommands.h
--- a/Engine/Source/Editor/CurveEditor/Public/CurveEditorCommands.h
+++ b/Engine/Source/Editor/CurveEditor/Public/CurveEditorCommands.h
@@ -90,7 +90,8 @@ public:
 
 	TSharedPtr<FUICommandInfo> AddKeyHovered;
 	TSharedPtr<FUICommandInfo> PasteKeysHovered;
-	TSharedPtr<FUICommandInfo> PasteOverwriteRange;
+	TSharedPtr<FUICommandInfo> PasteAndMerge;
+	TSharedPtr<FUICommandInfo> PasteRelative;
 
 	TSharedPtr<FUICommandInfo> AddKeyToAllCurves;
 
diff --git a/Engine/Source/Editor/CurveEditor/Public/CurveModel.h b/Engine/Source/Editor/CurveEditor/Public/CurveModel.h
--- a/Engine/Source/Editor/CurveEditor/Public/CurveModel.h
+++ b/Engine/Source/Editor/CurveEditor/Public/CurveModel.h
@@ -198,6 +198,19 @@ public:
 	 */
 	virtual void GetNeighboringKeys(const FKeyHandle InKeyHandle, TOptional<FKeyHandle>& OutPreviousKeyHandle, TOptional<FKeyHandle>& OutNextKeyHandle) const = 0;
 
+	/**
+	 * Finds the key before and after InTime.
+	 * 
+	 * If there is a key at InTime, then OutPreviousKeyHandle and OutNextKeyHandle are set to the same value. If there are multiple keys at InTime
+	 * (stacked), the implementation can return an arbitrary one as long as it is consistent (e.g. always the "first one") provided the internal state
+	 * has not changed.
+	 *
+	 * @param InTime The time near which to search for keys
+	 * @param OutPreviousKeyHandle The key with time <= InTime that is closest to InTime. Unset if there is no such key.
+	 * @param OutNextKeyHandle The key with time >= InTime that is closest to InTime. Unset if there is no such key.
+	 */
+	virtual void GetClosestKeysTo(double InTime, TOptional<FKeyHandle>& OutPreviousKeyHandle, TOptional<FKeyHandle>& OutNextKeyHandle) const;
+
 	/**
 	 * Get the interpolation mode to use at a specified time
 	 *
diff --git a/Engine/Source/Editor/CurveEditor/Public/Misc/KeyPasteArgs.h b/Engine/Source/Editor/CurveEditor/Public/Misc/KeyPasteArgs.h
new file mode 100644
--- /dev/null
+++ b/Engine/Source/Editor/CurveEditor/Public/Misc/KeyPasteArgs.h
@@ -0,0 +1,48 @@
+﻿// Copyright Epic Games, Inc. All Rights Reserved.
+
+#pragma once
+
+#include "Containers/Set.h"
+#include "CurveEditorTypes.h"
+#include "HAL/Platform.h"
+#include "Misc/EnumClassFlags.h"
+#include "Internationalization/Text.h"
+
+namespace UE::CurveEditor
+{
+
+/** Defines the merging behaviour of the paste operation. */
+enum class ECurveEditorPasteMode : uint8
+{
+	/** Pastes the keys from the clipboard overwriting any key in destination track between the first and last pasted keys */
+	OverwriteRange,
+	/** Existing keys are mostly retained. Existing keys are only replaced by pasted keys that have the same X value. */
+	Merge,
+};
+
+/** Enhances the paste operation */
+enum class ECurveEditorPasteFlags : uint8
+{
+	None,
+
+	/** Set the selection to the pasted keys. */
+	SetSelection = 1 << 0,
+	/** Pastes the keys from the clipboard aligning them to the nearest key to the left of the scrubber */
+	Relative = 1 << 1,
+
+	Default = SetSelection
+};
+ENUM_CLASS_FLAGS(ECurveEditorPasteFlags);
+
+/** Arguments for pasting keys in FCurveEditor. */
+struct FKeyPasteArgs
+{
+	/** Only the curve model IDs specified in this set. If empty, paste all in the clipboard. */
+	TSet<FCurveModelID> CurveModelIds;
+	ECurveEditorPasteMode Mode = ECurveEditorPasteMode::OverwriteRange;
+	ECurveEditorPasteFlags Flags = ECurveEditorPasteFlags::Default;
+
+	/** The name of the transaction */
+	FText OverrideTransactionName = NSLOCTEXT("CurveEditor", "PasteKeys", "Paste Keys");
+};
+}
