Curve Editor: Changed smart snap to Ctrl+h.

diff --git a/Engine/Source/Editor/CurveEditor/Private/CurveEditorCommands.cpp b/Engine/Source/Editor/CurveEditor/Private/CurveEditorCommands.cpp
--- a/Engine/Source/Editor/CurveEditor/Private/CurveEditorCommands.cpp
+++ b/Engine/Source/Editor/CurveEditor/Private/CurveEditorCommands.cpp
@@ -36,7 +36,7 @@ void FCurveEditorCommands::RegisterCommands()
 	UI_COMMAND(FlattenTangents, "Flatten", "Flatten Tangents", EUserInterfaceActionType::Button, FInputChord(EKeys::Six));
 	UI_COMMAND(StraightenTangents, "Straighten", "Straighten Tangents", EUserInterfaceActionType::Button, FInputChord());
 	
-	UI_COMMAND(SmartSnapKeys, "Smart Snap", "Snaps selected keys to the closest whole frame.\nThe key is placed where the curve intersects with the frame (imagine vertical line going down from the frame).", EUserInterfaceActionType::Button, FInputChord(EKeys::Seven));
+	UI_COMMAND(SmartSnapKeys, "Smart Snap", "Snaps selected keys to the closest whole frame.\nThe key is placed where the curve intersects with the frame (imagine vertical line going down from the frame).", EUserInterfaceActionType::Button, FInputChord(EModifierKey::Control, EKeys::H));
 
 	UI_COMMAND(MatchLastTangentToFirst, "Match Last Tangent To First", "Match last keys tangents to the first keys tangents", EUserInterfaceActionType::Button, FInputChord());
 	UI_COMMAND(MatchFirstTangentToLast, "Match First Tangent To Last", "Match first keys tangents to the last keys tangents", EUserInterfaceActionType::Button, FInputChord());
