UPDATE: changed curve flip feature in curve editor to be hidden by default.

diff --git a/Engine/Plugins/FX/Niagara/Source/NiagaraEditorWidgets/Private/SNiagaraCurveOverview.cpp b/Engine/Plugins/FX/Niagara/Source/NiagaraEditorWidgets/Private/SNiagaraCurveOverview.cpp
--- a/Engine/Plugins/FX/Niagara/Source/NiagaraEditorWidgets/Private/SNiagaraCurveOverview.cpp
+++ b/Engine/Plugins/FX/Niagara/Source/NiagaraEditorWidgets/Private/SNiagaraCurveOverview.cpp
@@ -477,6 +477,7 @@ void SNiagaraCurveOverview::Construct(const FArguments& InArgs, TSharedRef<FNiag
 	CurveEditor->GetTree()->SetSortPredicate(FCurveEditorTree::FTreeItemSortPredicate::CreateStatic(&FNiagaraCurveOverviewTreeItem::CompareTreeItems));
 
 	CurveEditorPanel = SNew(SCurveEditorPanel, CurveEditor.ToSharedRef())
+		.ShowCurveFlipButtons(true)
 		.MinimumViewPanelHeight(0.0f)
 		.TreeContent()
 		[
diff --git a/Engine/Source/Editor/CurveEditor/Private/SCurveEditorPanel.cpp b/Engine/Source/Editor/CurveEditor/Private/SCurveEditorPanel.cpp
--- a/Engine/Source/Editor/CurveEditor/Private/SCurveEditorPanel.cpp
+++ b/Engine/Source/Editor/CurveEditor/Private/SCurveEditorPanel.cpp
@@ -147,6 +147,7 @@ void SCurveEditorPanel::Construct(const FArguments& InArgs, TSharedRef<FCurveEdi
 	GridLineTintAttribute = InArgs._GridLineTint;
 	DisabledTimeSnapTooltipAttribute = InArgs._DisabledTimeSnapTooltip;
 	WeakTabManager = InArgs._TabManager;
+	ShowCurveFlipButtons = InArgs._ShowCurveFlipButtons;
 
 	CachedSelectionSerialNumber = 0;
 
@@ -1457,39 +1458,42 @@ TSharedPtr<FExtender> SCurveEditorPanel::GetToolbarExtender()
 			}
 			ToolBarBuilder.EndSection();
 
-			ToolBarBuilder.BeginSection("Flip Curve");
+			if (InEditorPanel->ShowCurveFlipButtons)
 			{
-				ToolBarBuilder.AddToolBarButton(
-					FCurveEditorCommands::Get().FlipCurveHorizontal,
-					NAME_None,
-					TAttribute<FText>(),
-					TAttribute<FText>(),
-					FSlateIcon(FAppStyle::Get().GetStyleSetName(), "GenericCurveEditor.FlipCurveHorizontal"));
-
-				ToolBarBuilder.AddComboButton(
-					FUIAction(),
-					FOnGetContent::CreateSP(InEditorPanel, &SCurveEditorPanel::MakeFlipCurveMenu, FCurveEditor::ECurveFlipDirection::Horizontal),
-					LOCTEXT("FlipCurveHorizontalOptions", "Flip Curve Horizontal"),
-					LOCTEXT("FlipCurveHorizontalOptionsToolTip", "Choose the range for horizontal curve flip"),
-					TAttribute<FSlateIcon>(),
-					true);
-
-				ToolBarBuilder.AddToolBarButton(
-					FCurveEditorCommands::Get().FlipCurveVertical,
-					NAME_None,
-					TAttribute<FText>(),
-					TAttribute<FText>(),
-					FSlateIcon(FAppStyle::Get().GetStyleSetName(), "GenericCurveEditor.FlipCurveVertical"));
-
-				ToolBarBuilder.AddComboButton(
-					FUIAction(),
-					FOnGetContent::CreateSP(InEditorPanel, &SCurveEditorPanel::MakeFlipCurveMenu, FCurveEditor::ECurveFlipDirection::Vertical),
-					LOCTEXT("FlipCurveVerticalOptions", "Flip Curve Vertical"),
-					LOCTEXT("FlipCurveVerticalOptionsToolTip", "Choose the range for vertical curve flip"),
-					TAttribute<FSlateIcon>(),
-					true);
+				ToolBarBuilder.BeginSection("Flip Curve");
+				{
+					ToolBarBuilder.AddToolBarButton(
+						FCurveEditorCommands::Get().FlipCurveHorizontal,
+						NAME_None,
+						TAttribute<FText>(),
+						TAttribute<FText>(),
+						FSlateIcon(FAppStyle::Get().GetStyleSetName(), "GenericCurveEditor.FlipCurveHorizontal"));
+
+					ToolBarBuilder.AddComboButton(
+						FUIAction(),
+						FOnGetContent::CreateSP(InEditorPanel, &SCurveEditorPanel::MakeFlipCurveMenu, FCurveEditor::ECurveFlipDirection::Horizontal),
+						LOCTEXT("FlipCurveHorizontalOptions", "Flip Curve Horizontal"),
+						LOCTEXT("FlipCurveHorizontalOptionsToolTip", "Choose the range for horizontal curve flip"),
+						TAttribute<FSlateIcon>(),
+						true);
+
+					ToolBarBuilder.AddToolBarButton(
+						FCurveEditorCommands::Get().FlipCurveVertical,
+						NAME_None,
+						TAttribute<FText>(),
+						TAttribute<FText>(),
+						FSlateIcon(FAppStyle::Get().GetStyleSetName(), "GenericCurveEditor.FlipCurveVertical"));
+
+					ToolBarBuilder.AddComboButton(
+						FUIAction(),
+						FOnGetContent::CreateSP(InEditorPanel, &SCurveEditorPanel::MakeFlipCurveMenu, FCurveEditor::ECurveFlipDirection::Vertical),
+						LOCTEXT("FlipCurveVerticalOptions", "Flip Curve Vertical"),
+						LOCTEXT("FlipCurveVerticalOptionsToolTip", "Choose the range for vertical curve flip"),
+						TAttribute<FSlateIcon>(),
+						true);
+				}
+				ToolBarBuilder.EndSection();
 			}
-			ToolBarBuilder.EndSection();
 
 			ToolBarBuilder.BeginSection("Tangents");
 			{
diff --git a/Engine/Source/Editor/CurveEditor/Public/SCurveEditorPanel.h b/Engine/Source/Editor/CurveEditor/Public/SCurveEditorPanel.h
--- a/Engine/Source/Editor/CurveEditor/Public/SCurveEditorPanel.h
+++ b/Engine/Source/Editor/CurveEditor/Public/SCurveEditorPanel.h
@@ -58,6 +58,7 @@ class CURVEEDITOR_API SCurveEditorPanel : public SCompoundWidget
 {
 	SLATE_BEGIN_ARGS(SCurveEditorPanel)
 		: _GridLineTint(FLinearColor(0.1f, 0.1f, 0.1f, 1.f))
+		, _ShowCurveFlipButtons(false)
 		, _MinimumViewPanelHeight(300.0f)
 	{}
 
@@ -70,6 +71,9 @@ class CURVEEDITOR_API SCurveEditorPanel : public SCompoundWidget
 		/** Optional Time Slider Controller which allows us to synchronize with an externally controlled Time Slider */
 		SLATE_ARGUMENT(TSharedPtr<ITimeSliderController>, ExternalTimeSliderController)
 
+		/** Optional Curve Flip functionality */
+		SLATE_ARGUMENT(bool, ShowCurveFlipButtons)
+
 		/** If specified, causes the time snap adjustment UI controls to be disabled and specifies the tooltip to be used. Can be used to disable time snap controls when externally controlled. */
 		SLATE_ATTRIBUTE(FText, DisabledTimeSnapTooltip)
 
@@ -205,6 +209,8 @@ private:
 	bool IsInlineEditPanelEditable() const;
 	EVisibility ShouldInstructionOverlayBeVisible() const;
 
+	bool ShowCurveFlipButtons;
+
 private:
 	
 	TSharedRef<SWidget> MakeToolsComboMenu(TSharedPtr<FExtender> InExtender);
