Curve Editor: Undo now allows restores FCurveEditorSelection::SerialNumber. Also fixes a bug where clearing selection would generate a transaction even though selection was already empty.

diff --git a/Engine/Source/Editor/CurveEditor/Private/CurveEditorSelection.cpp b/Engine/Source/Editor/CurveEditor/Private/CurveEditorSelection.cpp
--- a/Engine/Source/Editor/CurveEditor/Private/CurveEditorSelection.cpp
+++ b/Engine/Source/Editor/CurveEditor/Private/CurveEditorSelection.cpp
@@ -6,6 +6,7 @@
 #include "CurveEditor.h"
 #include "CurveEditorTypes.h"
 #include "CurveModel.h"
+#include "Algo/AllOf.h"
 #include "Misc/EnumClassFlags.h"
 #include "Templates/Tuple.h"
 #include "Templates/UniquePtr.h"
@@ -23,8 +24,13 @@ FCurveEditorSelection::FCurveEditorSelection(TWeakPtr<FCurveEditor> InWeakCurveE
 
 void FCurveEditorSelection::Clear()
 {
-	CurveToSelectedKeys.Reset();
-	++SerialNumber;
+	const bool bIsAlreadyEmpty = CurveToSelectedKeys.IsEmpty()
+		|| Algo::AllOf(CurveToSelectedKeys, [](const TPair<FCurveModelID, FKeyHandleSet>& Pair){ return Pair.Value.Num() == 0; }); 
+	if (!bIsAlreadyEmpty)
+	{
+		CurveToSelectedKeys.Reset();
+		++SerialNumber;
+	}
 }
 
 void FCurveEditorSelection::SerializeSelection(FArchive& Ar)
diff --git a/Engine/Source/Editor/CurveEditor/Private/CurveEditorSelectionPrivate.h b/Engine/Source/Editor/CurveEditor/Private/CurveEditorSelectionPrivate.h
--- a/Engine/Source/Editor/CurveEditor/Private/CurveEditorSelectionPrivate.h
+++ b/Engine/Source/Editor/CurveEditor/Private/CurveEditorSelectionPrivate.h
@@ -20,4 +20,17 @@ class FAddInternal
 		Selection.AddInternal(CurveID, PointType, Keys, bIncrementSerialNumber);
 	}
 };
+
+/** Implements passkey pattern to restrict access to FCurveEditorSelection::SerialNumber. */
+class FSetSerialNumber
+{
+	friend void SelectionDiffDetail::SetSerialNumberInternal(FCurveEditorSelection&, uint32);
+	
+	explicit FSetSerialNumber(
+		FCurveEditorSelection& Selection, int32 NewSerialNumber
+		)
+	{
+		Selection.SerialNumber = NewSerialNumber;
+	}
+};
 }
diff --git a/Engine/Source/Editor/CurveEditor/Private/Modification/Changes/SelectionDeltaChange.h b/Engine/Source/Editor/CurveEditor/Private/Modification/Changes/SelectionDeltaChange.h
--- a/Engine/Source/Editor/CurveEditor/Private/Modification/Changes/SelectionDeltaChange.h
+++ b/Engine/Source/Editor/CurveEditor/Private/Modification/Changes/SelectionDeltaChange.h
@@ -29,7 +29,12 @@ struct FSelectionDeltaChange
 	/** The per curve changes. */
 	TMap<FCurveModelID, FCurveSelectionDeltaChange> ChangedCurves;
 
-	bool HasChanges() const { return !ChangedCurves.IsEmpty(); }
+	/** Serial number of selection before the change. */
+	uint32 OldSerialNumber = 0;
+	/** Serial number of selection after the change. */
+	uint32 NewSerialNumber = 0;
+
+	bool HasChanges() const { return !ChangedCurves.IsEmpty() || OldSerialNumber != NewSerialNumber; }
 	operator bool() const { return HasChanges(); }
 };
 }
diff --git a/Engine/Source/Editor/CurveEditor/Private/Modification/Utils/SelectionChangeUtils.cpp b/Engine/Source/Editor/CurveEditor/Private/Modification/Utils/SelectionChangeUtils.cpp
--- a/Engine/Source/Editor/CurveEditor/Private/Modification/Utils/SelectionChangeUtils.cpp
+++ b/Engine/Source/Editor/CurveEditor/Private/Modification/Utils/SelectionChangeUtils.cpp
@@ -68,6 +68,11 @@ void AddKeysInternal(FCurveEditorSelection& Selection, FCurveModelID CurveID, EC
 {
 	KeySelection::FAddInternal(Selection, CurveID, PointType, Keys);
 }
+
+void SetSerialNumberInternal(FCurveEditorSelection& Selection, uint32 NewSerialNumber)
+{
+	KeySelection::FSetSerialNumber(Selection, NewSerialNumber);
+}
 	
 /** Removes keys from the array InGetRemovedKeys returns, and adds keys to the array InGetAddedKeys returns. */
 template<CGetKeys TGetAddedKeys, CGetKeys TGetRemovedKeys>
@@ -110,6 +115,9 @@ FSelectionDeltaChange DiffSelection(const FCurveEditorSelection& InOriginal, con
 		[](FCurveSelectionDeltaChange& CurveChanges) -> TArray<TPair<FKeyHandle, ECurvePointType>>&
 		{ return CurveChanges.AddedKeys; }
 	);
+
+	DeltaChange.OldSerialNumber = InOriginal.GetSerialNumber();
+	DeltaChange.NewSerialNumber = InTarget.GetSerialNumber();
 	
 	return DeltaChange;
 }
@@ -120,6 +128,8 @@ void ApplySelectionChange(FCurveEditorSelection& InOutChange, const FSelectionDe
 		[](const FCurveSelectionDeltaChange& CurveChange){ return CurveChange.AddedKeys; },
 		[](const FCurveSelectionDeltaChange& CurveChange){ return CurveChange.RemovedKeys; }
 		);
+
+	SelectionDiffDetail::SetSerialNumberInternal(InOutChange, InDeltaChange.NewSerialNumber);
 }
 	
 void RevertSelectionChange(FCurveEditorSelection& InOutChange, const FSelectionDeltaChange& InDeltaChange)
@@ -129,5 +139,7 @@ void RevertSelectionChange(FCurveEditorSelection& InOutChange, const FSelectionD
 		[](const FCurveSelectionDeltaChange& CurveChange){ return CurveChange.RemovedKeys; },
 		[](const FCurveSelectionDeltaChange& CurveChange){ return CurveChange.AddedKeys; }
 		);
+	
+	SelectionDiffDetail::SetSerialNumberInternal(InOutChange, InDeltaChange.OldSerialNumber);
 }
 }
diff --git a/Engine/Source/Editor/CurveEditor/Private/Modification/Utils/SelectionChangeUtils.h b/Engine/Source/Editor/CurveEditor/Private/Modification/Utils/SelectionChangeUtils.h
--- a/Engine/Source/Editor/CurveEditor/Private/Modification/Utils/SelectionChangeUtils.h
+++ b/Engine/Source/Editor/CurveEditor/Private/Modification/Utils/SelectionChangeUtils.h
@@ -23,4 +23,5 @@ void RevertSelectionChange(FCurveEditorSelection& InOutChange, const FSelectionD
 
 // Private - needed for FAddInternal.
 namespace SelectionDiffDetail { void AddKeysInternal(FCurveEditorSelection&, FCurveModelID, ECurvePointType, TArrayView<const FKeyHandle>); }
+namespace SelectionDiffDetail { void SetSerialNumberInternal(FCurveEditorSelection& Selection, uint32 NewSerialNumber); }
 }
diff --git a/Engine/Source/Editor/CurveEditor/Public/CurveEditorSelection.h b/Engine/Source/Editor/CurveEditor/Public/CurveEditorSelection.h
--- a/Engine/Source/Editor/CurveEditor/Public/CurveEditorSelection.h
+++ b/Engine/Source/Editor/CurveEditor/Public/CurveEditorSelection.h
@@ -16,6 +16,7 @@
 
 class FCurveEditor;
 namespace UE::CurveEditor::KeySelection { class FAddInternal; }
+namespace UE::CurveEditor::KeySelection { class FSetSerialNumber; }
 
 struct FCurvePointHandle;
 
@@ -87,6 +88,7 @@ private:
 struct CURVEEDITOR_API FCurveEditorSelection
 {
 	friend UE::CurveEditor::KeySelection::FAddInternal;
+	friend UE::CurveEditor::KeySelection::FSetSerialNumber;
 	
 	/**
 	 * Default constructor
