Curve Editor: Improve space of toolbar.
- Add clipping priorities to toolbar items of curve editor. The priorities are defined in the style sheet.
- Reduce tween slider width from 250 to 200
- Reduce key details width from 180 to 130

diff --git a/Engine/Plugins/Animation/TweeningUtils/Source/TweeningUtilsEditor/Private/Widgets/MVC/TweenToolbarController.cpp b/Engine/Plugins/Animation/TweeningUtils/Source/TweeningUtilsEditor/Private/Widgets/MVC/TweenToolbarController.cpp
--- a/Engine/Plugins/Animation/TweeningUtils/Source/TweeningUtilsEditor/Private/Widgets/MVC/TweenToolbarController.cpp
+++ b/Engine/Plugins/Animation/TweeningUtils/Source/TweeningUtilsEditor/Private/Widgets/MVC/TweenToolbarController.cpp
@@ -105,7 +105,8 @@ void FTweenToolbarController::AddToToolbar(FToolBarBuilder& ToolBarBuilder, FMak
 	   FOnGetContent::CreateRaw(this, &FTweenToolbarController::MakeTweenModeMenu),
 	   TAttribute<FText>::CreateRaw(this, &FTweenToolbarController::GetLabelForComboBox),
 	   TAttribute<FText>::CreateRaw(this, &FTweenToolbarController::GetToolTipForComboBox),
-	   TAttribute<FSlateIcon>::CreateRaw(this, &FTweenToolbarController::GetIconForComboBox)
+	   TAttribute<FSlateIcon>::CreateRaw(this, &FTweenToolbarController::GetIconForComboBox),
+	   false, NAME_None, {}, {}, {}, EUserInterfaceActionType::Button, InArgs.FunctionSelectResizeParams
 	   );
 	ToolBarBuilder.SetLabelVisibility(EVisibility::Collapsed);
 		
@@ -114,12 +115,14 @@ void FTweenToolbarController::AddToToolbar(FToolBarBuilder& ToolBarBuilder, FMak
 		.TweenModel(MakeSelectedTweenModelAttr())
 		.SliderIcon(TAttribute<const FSlateBrush*>::CreateRaw(this, &FTweenToolbarController::GetIconForSlider))
 		.SliderColor(TAttribute<FLinearColor>::CreateRaw(this, &FTweenToolbarController::GetColorForSlider))
-		.OverrideSliderPosition(MoveTemp(InArgs.OverrideSliderPositionAttr))
+		.OverrideSliderPosition(MoveTemp(InArgs.OverrideSliderPositionAttr)),
+		{}, NAME_None, true, {}, {}, InArgs.SliderResizeParams
 		);
 	
 	ToolBarBuilder.AddToolBarButton(FTweeningUtilsCommands::Get().ToggleOvershootMode,
 		NAME_None, TAttribute<FText>(), TAttribute<FText>(),
-		TAttribute<FSlateIcon>::CreateRaw(this, &FTweenToolbarController::GetOvershootModeIcon)
+		TAttribute<FSlateIcon>::CreateRaw(this, &FTweenToolbarController::GetOvershootModeIcon),
+		NAME_None, {}, {}, {}, InArgs.OvershootResizeParams
 		);
 }
 
diff --git a/Engine/Plugins/Animation/TweeningUtils/Source/TweeningUtilsEditor/Private/Widgets/TweenSliderStyle.cpp b/Engine/Plugins/Animation/TweeningUtils/Source/TweeningUtilsEditor/Private/Widgets/TweenSliderStyle.cpp
--- a/Engine/Plugins/Animation/TweeningUtils/Source/TweeningUtilsEditor/Private/Widgets/TweenSliderStyle.cpp
+++ b/Engine/Plugins/Animation/TweeningUtils/Source/TweeningUtilsEditor/Private/Widgets/TweenSliderStyle.cpp
@@ -31,7 +31,7 @@ void FTweenPointStyle::GetResources(TArray<const FSlateBrush*>& OutBrushes) cons
 }
 
 FTweenSliderStyle::FTweenSliderStyle()
-	: BarDimensions(250.0, 12.0)
+	: BarDimensions(200.0, 12.0)
 	, BarBrush(static_cast<FSlateBrush>(FSlateColorBrush(FStyleColors::Black)))
 	, NormalSliderButton(static_cast<FSlateBrush>(FSlateRoundedBoxBrush(FStyleColors::ForegroundHover, 5.0f, FStyleColors::Transparent, 0.f)))
 	, HoveredSliderButton(static_cast<FSlateBrush>(FSlateRoundedBoxBrush(FStyleColors::ForegroundHover, 5.0f, FStyleColors::Transparent, 0.f)))
diff --git a/Engine/Plugins/Animation/TweeningUtils/Source/TweeningUtilsEditor/Public/Widgets/MVC/TweenToolbarController.h b/Engine/Plugins/Animation/TweeningUtils/Source/TweeningUtilsEditor/Public/Widgets/MVC/TweenToolbarController.h
--- a/Engine/Plugins/Animation/TweeningUtils/Source/TweeningUtilsEditor/Public/Widgets/MVC/TweenToolbarController.h
+++ b/Engine/Plugins/Animation/TweeningUtils/Source/TweeningUtilsEditor/Public/Widgets/MVC/TweenToolbarController.h
@@ -2,6 +2,7 @@
 
 #pragma once
 
+#include "Framework/MultiBox/MultiBoxDefs.h"
 #include "Misc/Attribute.h"
 #include "Misc/NotNull.h"
 #include "Templates/PimplPtr.h"
@@ -29,7 +30,7 @@ class FTweenModel;
  * Makes a toolbar as follows |Combo|Slider|Overshot|
  * - Combo: Selects the blend function
  * - Slider: Adjusts the blend value and updates the FTweenModel
- * - Overshot: Toggle button for changing the overshot property on FTweenModel
+ * - Overshoot: Toggle button for changing the overshoot property on FTweenModel
  */
 class TWEENINGUTILSEDITOR_API FTweenToolbarController : public FNoncopyable
 {
@@ -54,6 +55,13 @@ public:
 		 * Useful e.g. in conjunction with FTweenMouseSlidingController (while pressing U+LMB, give the user a visual association). 
 		 */
 		TAttribute<TOptional<float>> OverrideSliderPositionAttr;
+
+		/** Affects how the tween function selection combo box is resized. */
+		FMenuEntryResizeParams FunctionSelectResizeParams;
+		/** Affects how the slider widget is resized. */
+		FMenuEntryResizeParams SliderResizeParams;
+		/** Affects how the overshoot button is resized. */
+		FMenuEntryResizeParams OvershootResizeParams;
 	};
 	/**
 	 * Extends the toolbar. Safe to call multiple times (although typically you'll call this exactly once).
diff --git a/Engine/Plugins/Editor/CurveEditorTools/Source/CurveEditorTools/Private/Extensions/Tweening/TweenEditorExtension.cpp b/Engine/Plugins/Editor/CurveEditorTools/Source/CurveEditorTools/Private/Extensions/Tweening/TweenEditorExtension.cpp
--- a/Engine/Plugins/Editor/CurveEditorTools/Source/CurveEditorTools/Private/Extensions/Tweening/TweenEditorExtension.cpp
+++ b/Engine/Plugins/Editor/CurveEditorTools/Source/CurveEditorTools/Private/Extensions/Tweening/TweenEditorExtension.cpp
@@ -7,6 +7,7 @@
 #include "Framework/MultiBox/MultiBoxBuilder.h"
 #include "Framework/MultiBox/MultiBoxExtender.h"
 #include "Misc/Optional.h"
+#include "Misc/ResizeParamUtils.h"
 #include "Widgets/MVC/TweenToolbarController.h"
 
 namespace UE::CurveEditorTools
@@ -41,7 +42,13 @@ TSharedPtr<FExtender> FTweenEditorExtension::MakeToolbarExtender(const TSharedRe
 				
 				ToolbarBuilder.BeginStyleOverride(TEXT("CurveEditorTweenToolbar")); // Forces min sizes for combo button
 				ToolbarBuilder.BeginSection("Tween");
-				TweenControllers->ToolbarController.AddToToolbar(ToolbarBuilder, { MoveTemp(OverrideSliderPosition) });
+				TweenControllers->ToolbarController.AddToToolbar(ToolbarBuilder,
+				{
+					.OverrideSliderPositionAttr = MoveTemp(OverrideSliderPosition),
+					.FunctionSelectResizeParams = CurveEditor::MakeResizeParams(TEXT("Tween.FunctionSelect")),
+					.SliderResizeParams			= CurveEditor::MakeResizeParams(TEXT("Tween.Slider")),
+					.OvershootResizeParams		= CurveEditor::MakeResizeParams(TEXT("Tween.Overshoot")) 
+				});
 				ToolbarBuilder.EndSection();
 				ToolbarBuilder.EndStyleOverride();
 			}
diff --git a/Engine/Source/Editor/CurveEditor/Private/Filters/PromotedFilterContainer.cpp b/Engine/Source/Editor/CurveEditor/Private/Filters/PromotedFilterContainer.cpp
--- a/Engine/Source/Editor/CurveEditor/Private/Filters/PromotedFilterContainer.cpp
+++ b/Engine/Source/Editor/CurveEditor/Private/Filters/PromotedFilterContainer.cpp
@@ -30,11 +30,11 @@ FPromotedFilterContainer::~FPromotedFilterContainer()
 	FInputBindingManager::Get().RemoveContextByName(CommandContext->GetContextName());
 }
 
-void FPromotedFilterContainer::AppendToBuilder(FToolBarBuilder& InToolBarBuilder) const
+void FPromotedFilterContainer::AppendToBuilder(FToolBarBuilder& InToolBarBuilder, const FMenuEntryResizeParams& InResizeParams) const
 {
 	for (const FFilterData& FilterData : PromotedFilters)
 	{
-		InToolBarBuilder.AddToolBarButton(FilterData.Command);
+		InToolBarBuilder.AddToolBarButton(FilterData.Command, NAME_None, {}, {}, {}, NAME_None, {}, {}, {}, InResizeParams);
 	}
 }
 
diff --git a/Engine/Source/Editor/CurveEditor/Private/SCurveEditorPanel.cpp b/Engine/Source/Editor/CurveEditor/Private/SCurveEditorPanel.cpp
--- a/Engine/Source/Editor/CurveEditor/Private/SCurveEditorPanel.cpp
+++ b/Engine/Source/Editor/CurveEditor/Private/SCurveEditorPanel.cpp
@@ -66,6 +66,7 @@
 #include "SlotBase.h"
 #include "Filters/CurveEditorEulerFilter.h"
 #include "Filters/PromotedFilterContainer.h"
+#include "Misc/ResizeParamUtils.h"
 #include "Modification/Utils/ScopedSelectionTransaction.h"
 #include "Styling/AppStyle.h"
 #include "Styling/CoreStyle.h"
@@ -1341,13 +1342,43 @@ const FGeometry& SCurveEditorPanel::GetViewContainerGeometry() const
 	return CurveViewsContainer->GetCachedGeometry();
 }
 
+namespace UE::CurveEditor::PanelDetail
+{
+	
+static void FillToolbar(
+	FToolBarBuilder& ToolBarBuilder,
+	TSharedRef<SCurveKeyDetailPanel> InKeyDetailsPanel,
+	TSharedRef<SCurveEditorPanel> InEditorPanel,
+	TSharedRef<FPromotedFilterContainer> InToolbarPromotedFilters,
+	TSharedPtr<FExtender> InExtender
+	)
+{
+	
+} 
+}
+
 TSharedPtr<FExtender> SCurveEditorPanel::GetToolbarExtender()
 {
+	using namespace UE::CurveEditor::PanelDetail;
+	
 	// We're going to create a new Extender and add the main Curve Editor icons to it.
 	// We combine this with the extender provided by the Curve Editor Module as that extender has been extended by tools
+	TSharedPtr<FExtender> Extender = CombineEditorExtensions();
+
+	Extender->AddToolBarExtension(
+		"Asset", EExtensionHook::After, GetCommands(),
+		FToolBarExtensionDelegate::CreateSP(this, &SCurveEditorPanel::BuildToolbar, Extender)
+	);
+
+	return Extender;
+}
+
+TSharedPtr<FExtender> SCurveEditorPanel::CombineEditorExtensions() const
+{
+	const TSharedRef<FUICommandList>& CommandListRef = GetCommands().ToSharedRef();
+	
 	ICurveEditorModule& CurveEditorModule = FModuleManager::Get().LoadModuleChecked<ICurveEditorModule>("CurveEditor");
 	TArray<TSharedPtr<FExtender>> ToolbarExtenders;
-	const TSharedRef<FUICommandList> CommandListRef = GetCommands().ToSharedRef();
 	for (ICurveEditorModule::FCurveEditorMenuExtender& ExtenderCallback : CurveEditorModule.GetAllToolBarMenuExtenders())
 	{
 		ToolbarExtenders.Add(ExtenderCallback.Execute(CommandListRef));
@@ -1359,196 +1390,190 @@ TSharedPtr<FExtender> SCurveEditorPanel::GetToolbarExtender()
 			ToolbarExtenders.Add(OptionalExtender);
 		}
 	}
-	TSharedPtr<FExtender> Extender = FExtender::Combine(ToolbarExtenders);
-
-	struct Local
-	{
-		static void FillToolbar(
-			FToolBarBuilder& ToolBarBuilder,
-			TSharedRef<SCurveKeyDetailPanel> InKeyDetailsPanel,
-			TSharedRef<SCurveEditorPanel> InEditorPanel,
-			TSharedRef<UE::CurveEditor::FPromotedFilterContainer> InToolbarPromotedFilters,
-			TSharedPtr<FExtender> InExtender
-			)
-		{
-			ToolBarBuilder.BeginSection("View");
-			ToolBarBuilder.BeginStyleOverride("CurveEditorToolbar");
-			{
-				ToolBarBuilder.SetLabelVisibility(EVisibility::Visible); // Show label..
-				ToolBarBuilder.AddComboButton(
-					FUIAction(),
-					FOnGetContent::CreateSP(InEditorPanel, &SCurveEditorPanel::MakeCurveEditorCurveViewOptionsMenu, InExtender),
-					LOCTEXT("CurveEditor.ViewOptions.Label", "View"),
-					LOCTEXT("CurveEditor.ViewOptions.ToolTip", "View Options"),
-					FSlateIcon(FAppStyle::Get().GetStyleSetName(), "Icons.Visibility")
-					);
-				ToolBarBuilder.SetLabelVisibility(EVisibility::Collapsed); // ... but don't affect the other entries
-			}
-			ToolBarBuilder.EndSection();
-
-			ToolBarBuilder.BeginSection("Tools");
-			{
-				ToolBarBuilder.BeginStyleOverride("CurveEditorToolBar.ToolsCombo"); // Min size for combo button so size stays consistent independent of label size
-				ToolBarBuilder.SetLabelVisibility(EVisibility::Visible); // Show label..
-				ToolBarBuilder.AddComboButton(
-					FUIAction(),
-					FOnGetContent::CreateSP(InEditorPanel, &SCurveEditorPanel::MakeToolsComboMenu, InExtender),
-					TAttribute<FText>(InEditorPanel, &SCurveEditorPanel::GetCurrentToolLabel),
-					FText::GetEmpty(),
-					TAttribute<FSlateIcon>(InEditorPanel, &SCurveEditorPanel::GetCurrentToolIcon)
-					);
-				ToolBarBuilder.SetLabelVisibility(EVisibility::Collapsed); // ... but don't affect the other entries
-				ToolBarBuilder.EndStyleOverride();
-			}
-			ToolBarBuilder.EndSection();
-
-			ToolBarBuilder.BeginSection("Key Details");
-			{
-				ToolBarBuilder.AddWidget(InKeyDetailsPanel);
-			}
-			ToolBarBuilder.EndSection();
-
-			ToolBarBuilder.BeginSection("Adjustment");
-			{
-				// Dropdown Menu for choosing your axis snapping for tool movement
-				TAttribute<FSlateIcon> AxisSnappingModeIcon;
-				AxisSnappingModeIcon.Bind(TAttribute<FSlateIcon>::FGetter::CreateLambda([InEditorPanel] {
-					switch (InEditorPanel->GetCurveEditor()->GetAxisSnap().RestrictedAxisList)
-					{
-					case ECurveEditorSnapAxis::CESA_X:
-						return FCurveEditorCommands::Get().SetAxisSnappingHorizontal->GetIcon();
-					case ECurveEditorSnapAxis::CESA_Y:
-						return FCurveEditorCommands::Get().SetAxisSnappingVertical->GetIcon();
-					default: // ECurveEditorSnapAxis::CESA_None
-						return FCurveEditorCommands::Get().SetAxisSnappingNone->GetIcon();
-					}
-				}));
-
-				ToolBarBuilder.AddComboButton(
-					FUIAction(),
-					FOnGetContent::CreateSP(InEditorPanel, &SCurveEditorPanel::MakeAxisSnapMenu),
-					LOCTEXT("AxisSnappingOptions", "Axis Snapping"),
-					LOCTEXT("AxisSnappingOptionsToolTip", "Choose which axes movement tools are locked to."),
-					AxisSnappingModeIcon);
-				
-				// Toggle Button for Time Snapping
-				ToolBarBuilder.AddToolBarButton(FCurveEditorCommands::Get().ToggleInputSnapping);
-
-				// Dropdown Menu to choose the snapping scale.
-				FUIAction TimeSnapMenuAction(FExecuteAction(), FCanExecuteAction::CreateLambda([InEditorPanel] { return !InEditorPanel->DisabledTimeSnapTooltipAttribute.IsSet(); }));
-				ToolBarBuilder.AddComboButton(
-					TimeSnapMenuAction,
-					FOnGetContent::CreateSP(InEditorPanel, &SCurveEditorPanel::MakeTimeSnapMenu),
-					LOCTEXT("TimeSnappingOptions", "Time Snapping"),
-					TAttribute<FText>(InEditorPanel, &SCurveEditorPanel::GetTimeSnapMenuTooltip),
-					TAttribute<FSlateIcon>(),
-					true);
-
-				// Toggle Button for Value Snapping
-				ToolBarBuilder.AddToolBarButton(FCurveEditorCommands::Get().ToggleOutputSnapping);
-
-				// Dropdown Menu to choose the snapping scale.
-				ToolBarBuilder.AddComboButton(
-					FUIAction(),
-					FOnGetContent::CreateSP(InEditorPanel, &SCurveEditorPanel::MakeGridSpacingMenu),
-					LOCTEXT("GridSnappingOptions", "Grid Snapping"),
-					LOCTEXT("GridSnappingOptionsToolTip", "Choose the spacing between horizontal grid lines."),
-					TAttribute<FSlateIcon>(),
-					true);
-			}
-			ToolBarBuilder.EndSection();
-
-			if (InEditorPanel->ShowCurveFlipButtons)
-			{
-				ToolBarBuilder.BeginSection("Flip Curve");
-				{
-					ToolBarBuilder.AddToolBarButton(
-						FCurveEditorCommands::Get().FlipCurveHorizontal,
-						NAME_None,
-						TAttribute<FText>(),
-						TAttribute<FText>(),
-						FSlateIcon(FAppStyle::Get().GetStyleSetName(), "GenericCurveEditor.FlipCurveHorizontal"));
-
-					ToolBarBuilder.AddComboButton(
-						FUIAction(),
-						FOnGetContent::CreateSP(InEditorPanel, &SCurveEditorPanel::MakeFlipCurveMenu, FCurveEditor::ECurveFlipDirection::Horizontal),
-						LOCTEXT("FlipCurveHorizontalOptions", "Flip Curve Horizontal"),
-						LOCTEXT("FlipCurveHorizontalOptionsToolTip", "Choose the range for horizontal curve flip"),
-						TAttribute<FSlateIcon>(),
-						true);
-
-					ToolBarBuilder.AddToolBarButton(
-						FCurveEditorCommands::Get().FlipCurveVertical,
-						NAME_None,
-						TAttribute<FText>(),
-						TAttribute<FText>(),
-						FSlateIcon(FAppStyle::Get().GetStyleSetName(), "GenericCurveEditor.FlipCurveVertical"));
-
-					ToolBarBuilder.AddComboButton(
-						FUIAction(),
-						FOnGetContent::CreateSP(InEditorPanel, &SCurveEditorPanel::MakeFlipCurveMenu, FCurveEditor::ECurveFlipDirection::Vertical),
-						LOCTEXT("FlipCurveVerticalOptions", "Flip Curve Vertical"),
-						LOCTEXT("FlipCurveVerticalOptionsToolTip", "Choose the range for vertical curve flip"),
-						TAttribute<FSlateIcon>(),
-						true);
-				}
-				ToolBarBuilder.EndSection();
-			}
-
-			ToolBarBuilder.BeginSection("Tangents");
-			{
-				ToolBarBuilder.BeginStyleOverride("CurveEditorToolBar.TangentMode"); // Min size for combo button so size stays consistent independent of label size
-				ToolBarBuilder.SetLabelVisibility(EVisibility::Visible); // Show label...
-				ToolBarBuilder.AddComboButton(
-					FUIAction(FExecuteAction(), FCanExecuteAction::CreateSP(InEditorPanel, &SCurveEditorPanel::IsTangentModeComboEnabled)),
-					FOnGetContent::CreateSP(InEditorPanel, &SCurveEditorPanel::MakeTangentModeMenu),
-					TAttribute<FText>::CreateSP(&InEditorPanel.Get(), &SCurveEditorPanel::GetTangentModeLabel),
-					TAttribute<FText>::CreateSP(&InEditorPanel.Get(), &SCurveEditorPanel::GetTangentModeTooltip),
-					TAttribute<FSlateIcon>::CreateSP(&InEditorPanel.Get(), &SCurveEditorPanel::GetTangentModeIcon)
-					);
-				ToolBarBuilder.SetLabelVisibility(EVisibility::Collapsed); // ... but don't affect the other entries
-				ToolBarBuilder.EndStyleOverride();
-				
-				const int32 SupportedTangentTypes = InEditorPanel->CurveEditor->GetSupportedTangentTypes();
-				if (SupportedTangentTypes & (int32)ECurveEditorTangentTypes::InterpolationCubicWeighted)
-				{
-					ToolBarBuilder.AddToolBarButton(FCurveEditorCommands::Get().InterpolationToggleWeighted);
-				}
-				ToolBarBuilder.AddToolBarButton(FCurveEditorCommands::Get().FlattenTangents);
-				ToolBarBuilder.AddToolBarButton(FCurveEditorCommands::Get().StraightenTangents);
-			}
-			ToolBarBuilder.EndSection();
+	return FExtender::Combine(ToolbarExtenders);
+}
 
-			ToolBarBuilder.BeginSection("Curves");
+void SCurveEditorPanel::BuildToolbar(FToolBarBuilder& InToolBarBuilder, TSharedPtr<FExtender> InBaseExtender)
+{
+	using namespace UE::CurveEditor;
+	
+	InToolBarBuilder.BeginSection("View");
+	InToolBarBuilder.BeginStyleOverride("CurveEditorToolbar");
+	{
+		InToolBarBuilder.SetLabelVisibility(EVisibility::Visible); // Show label..
+		InToolBarBuilder.AddComboButton(
+			FUIAction(),
+			FOnGetContent::CreateSP(this, &SCurveEditorPanel::MakeCurveEditorCurveViewOptionsMenu, InBaseExtender),
+			LOCTEXT("CurveEditor.ViewOptions.Label", "View"),
+			LOCTEXT("CurveEditor.ViewOptions.ToolTip", "View Options"),
+			FSlateIcon(FAppStyle::Get().GetStyleSetName(), "Icons.Visibility"),
+			false, NAME_None, {}, {}, {}, EUserInterfaceActionType::Button,
+			MakeResizeParams(TEXT("View"))
+			);
+		InToolBarBuilder.SetLabelVisibility(EVisibility::Collapsed); // ... but don't affect the other entries
+	}
+	InToolBarBuilder.EndSection();
+
+	InToolBarBuilder.BeginSection("Tools");
+	{
+		InToolBarBuilder.BeginStyleOverride("CurveEditorToolBar.ToolsCombo"); // Min size for combo button so size stays consistent independent of label size
+		InToolBarBuilder.SetLabelVisibility(EVisibility::Visible); // Show label..
+		InToolBarBuilder.AddComboButton(
+			FUIAction(),
+			FOnGetContent::CreateSP(this, &SCurveEditorPanel::MakeToolsComboMenu, InBaseExtender),
+			TAttribute<FText>(this, &SCurveEditorPanel::GetCurrentToolLabel),
+			FText::GetEmpty(),
+			TAttribute<FSlateIcon>(this, &SCurveEditorPanel::GetCurrentToolIcon),
+			false, NAME_None, {}, {}, {}, EUserInterfaceActionType::Button,
+			MakeResizeParams(TEXT("Tools"))
+			);
+		InToolBarBuilder.SetLabelVisibility(EVisibility::Collapsed); // ... but don't affect the other entries
+		InToolBarBuilder.EndStyleOverride();
+	}
+	InToolBarBuilder.EndSection();
+
+	InToolBarBuilder.BeginSection("Key Details");
+	{
+		InToolBarBuilder.AddWidget(KeyDetailsView.ToSharedRef(), {}, NAME_None, true, {}, {}, MakeResizeParams(TEXT("KeyDetails")));
+	}
+	InToolBarBuilder.EndSection();
+
+	InToolBarBuilder.BeginSection("Adjustment");
+	{
+		// Dropdown Menu for choosing your axis snapping for tool movement
+		TAttribute<FSlateIcon> AxisSnappingModeIcon;
+		AxisSnappingModeIcon.Bind(TAttribute<FSlateIcon>::FGetter::CreateLambda([this] {
+			switch (this->GetCurveEditor()->GetAxisSnap().RestrictedAxisList)
 			{
-				ToolBarBuilder.SetLabelVisibility(EVisibility::Visible); // Show label...
-				ToolBarBuilder.AddComboButton(
-					FUIAction(),
-					FOnGetContent::CreateSP(InEditorPanel, &SCurveEditorPanel::MakeCurvesMenu, InExtender),
-					LOCTEXT("CurveEditor.CurvesCombo.Label", "Curves"),
-					LOCTEXT("CurveEditor.CurvesCombo.ToolTip", "Curves"),
-					FSlateIcon(FAppStyle::Get().GetStyleSetName(), "GenericCurveEditor.Curves")
-					);
-				ToolBarBuilder.SetLabelVisibility(EVisibility::Collapsed); // ... but don't affect the other entries
-				
-				InToolbarPromotedFilters->AppendToBuilder(ToolBarBuilder);
-				ToolBarBuilder.AddToolBarButton(FCurveEditorCommands::Get().OpenUserImplementableFilterWindow);
+			case ECurveEditorSnapAxis::CESA_X:
+				return FCurveEditorCommands::Get().SetAxisSnappingHorizontal->GetIcon();
+			case ECurveEditorSnapAxis::CESA_Y:
+				return FCurveEditorCommands::Get().SetAxisSnappingVertical->GetIcon();
+			default: // ECurveEditorSnapAxis::CESA_None
+				return FCurveEditorCommands::Get().SetAxisSnappingNone->GetIcon();
 			}
-			
-			ToolBarBuilder.EndStyleOverride();
+		}));
+		InToolBarBuilder.AddComboButton(
+			FUIAction(),
+			FOnGetContent::CreateSP(this, &SCurveEditorPanel::MakeAxisSnapMenu),
+			LOCTEXT("AxisSnappingOptions", "Axis Snapping"),
+			LOCTEXT("AxisSnappingOptionsToolTip", "Choose which axes movement tools are locked to."),
+			AxisSnappingModeIcon,
+			false, NAME_None, {}, {}, {}, EUserInterfaceActionType::Button,
+			MakeResizeParams(TEXT("AxisSnapping"))
+			);
+		
+		// Toggle Button for Time Snapping
+		InToolBarBuilder.AddToolBarButton(
+			FCurveEditorCommands::Get().ToggleInputSnapping,
+			NAME_None, {}, {}, {}, NAME_None, {}, {}, {},
+			MakeResizeParams(TEXT("ToggleInputSnapping"))
+			);
+
+		// Dropdown Menu to choose the snapping scale.
+		FUIAction TimeSnapMenuAction(FExecuteAction(), FCanExecuteAction::CreateLambda([this] { return !this->DisabledTimeSnapTooltipAttribute.IsSet(); }));
+		InToolBarBuilder.AddComboButton(
+			TimeSnapMenuAction,
+			FOnGetContent::CreateSP(this, &SCurveEditorPanel::MakeTimeSnapMenu),
+			LOCTEXT("TimeSnappingOptions", "Time Snapping"),
+			TAttribute<FText>(this, &SCurveEditorPanel::GetTimeSnapMenuTooltip),
+			TAttribute<FSlateIcon>(),
+			true,
+			NAME_None, {}, {}, {}, EUserInterfaceActionType::Button,
+			MakeResizeParams(TEXT("TimeSnapping"))
+			);
+
+		// Toggle Button for Value Snapping
+		InToolBarBuilder.AddToolBarButton(FCurveEditorCommands::Get().ToggleOutputSnapping,
+			NAME_None, {}, {}, {}, NAME_None, {}, {}, {},
+			MakeResizeParams(TEXT("ToggleOutputSnapping"))
+			);
+
+		// Dropdown Menu to choose the snapping scale.
+		InToolBarBuilder.AddComboButton(
+			FUIAction(),
+			FOnGetContent::CreateSP(this, &SCurveEditorPanel::MakeGridSpacingMenu),
+			LOCTEXT("GridSnappingOptions", "Grid Snapping"),
+			LOCTEXT("GridSnappingOptionsToolTip", "Choose the spacing between horizontal grid lines."),
+			TAttribute<FSlateIcon>(),
+			true,
+			NAME_None, {}, {}, {}, EUserInterfaceActionType::Button,
+			MakeResizeParams(TEXT("GridSnapping"))
+			);
+	}
+	InToolBarBuilder.EndSection();
+
+	InToolBarBuilder.BeginSection("Tangents");
+	{
+		InToolBarBuilder.BeginStyleOverride("CurveEditorToolBar.TangentMode"); // Min size for combo button so size stays consistent independent of label size
+		InToolBarBuilder.SetLabelVisibility(EVisibility::Visible); // Show label...
+		InToolBarBuilder.AddComboButton(
+			FUIAction(FExecuteAction(), FCanExecuteAction::CreateSP(this, &SCurveEditorPanel::IsTangentModeComboEnabled)),
+			FOnGetContent::CreateSP(this, &SCurveEditorPanel::MakeTangentModeMenu),
+			TAttribute<FText>::CreateSP(this, &SCurveEditorPanel::GetTangentModeLabel),
+			TAttribute<FText>::CreateSP(this, &SCurveEditorPanel::GetTangentModeTooltip),
+			TAttribute<FSlateIcon>::CreateSP(this, &SCurveEditorPanel::GetTangentModeIcon),
+			false, NAME_None, {}, {}, {}, EUserInterfaceActionType::Button,
+			MakeResizeParams(TEXT("TangentMenu"))
+			);
+		InToolBarBuilder.SetLabelVisibility(EVisibility::Collapsed); // ... but don't affect the other entries
+		InToolBarBuilder.EndStyleOverride();
+		
+		const int32 SupportedTangentTypes = this->CurveEditor->GetSupportedTangentTypes();
+		if (SupportedTangentTypes & (int32)ECurveEditorTangentTypes::InterpolationCubicWeighted)
+		{
+			InToolBarBuilder.AddToolBarButton(FCurveEditorCommands::Get().InterpolationToggleWeighted,
+				NAME_None, {}, {}, {}, NAME_None, {}, {}, {},
+				MakeResizeParams(TEXT("InterpolationToggleWeighted"))
+				);
 		}
+		InToolBarBuilder.AddToolBarButton(FCurveEditorCommands::Get().FlattenTangents,
+			NAME_None, {}, {}, {}, NAME_None, {}, {}, {},
+			MakeResizeParams(TEXT("FlattenTangents"))
+			);
+		InToolBarBuilder.AddToolBarButton(FCurveEditorCommands::Get().StraightenTangents,
+			NAME_None, {}, {}, {}, NAME_None, {}, {}, {},
+			MakeResizeParams(TEXT("StraightenTangents"))
+			);
+	}
+	InToolBarBuilder.EndSection();
 
-	};
-	Extender->AddToolBarExtension(
-		"Asset",
-		EExtensionHook::After,
-		GetCommands(),
-		FToolBarExtensionDelegate::CreateStatic(
-			&Local::FillToolbar, KeyDetailsView.ToSharedRef(), SharedThis(this), GetCurveEditor()->GetToolbarPromotedFilters().ToSharedRef(), Extender
-			)
-	);
-
-	return Extender;
+	if (ShowCurveFlipButtons)
+	{
+		InToolBarBuilder.BeginSection("Flip Curve");
+		{
+			InToolBarBuilder.AddToolBarButton(
+				FCurveEditorCommands::Get().FlipCurveHorizontal,
+				NAME_None,
+				TAttribute<FText>(),
+				TAttribute<FText>(),
+				FSlateIcon(FAppStyle::Get().GetStyleSetName(), "GenericCurveEditor.FlipCurveHorizontal"));
+
+			InToolBarBuilder.AddComboButton(
+				FUIAction(),
+				FOnGetContent::CreateSP(this, &SCurveEditorPanel::MakeFlipCurveMenu, FCurveEditor::ECurveFlipDirection::Horizontal),
+				LOCTEXT("FlipCurveHorizontalOptions", "Flip Curve Horizontal"),
+				LOCTEXT("FlipCurveHorizontalOptionsToolTip", "Choose the range for horizontal curve flip"),
+				TAttribute<FSlateIcon>(),
+				true);
+
+			InToolBarBuilder.AddToolBarButton(
+				FCurveEditorCommands::Get().FlipCurveVertical,
+				NAME_None,
+				TAttribute<FText>(),
+				TAttribute<FText>(),
+				FSlateIcon(FAppStyle::Get().GetStyleSetName(), "GenericCurveEditor.FlipCurveVertical"));
+
+			InToolBarBuilder.AddComboButton(
+				FUIAction(),
+				FOnGetContent::CreateSP(this, &SCurveEditorPanel::MakeFlipCurveMenu, FCurveEditor::ECurveFlipDirection::Vertical),
+				LOCTEXT("FlipCurveVerticalOptions", "Flip Curve Vertical"),
+				LOCTEXT("FlipCurveVerticalOptionsToolTip", "Choose the range for vertical curve flip"),
+				TAttribute<FSlateIcon>(),
+				true);
+		}
+		InToolBarBuilder.EndSection();
+	}
 }
 
 TSharedRef<SWidget> SCurveEditorPanel::MakeTimeSnapMenu()
diff --git a/Engine/Source/Editor/CurveEditor/Public/Filters/PromotedFilterContainer.h b/Engine/Source/Editor/CurveEditor/Public/Filters/PromotedFilterContainer.h
--- a/Engine/Source/Editor/CurveEditor/Public/Filters/PromotedFilterContainer.h
+++ b/Engine/Source/Editor/CurveEditor/Public/Filters/PromotedFilterContainer.h
@@ -3,12 +3,10 @@
 #pragma once
 
 #include "Delegates/Delegate.h"
-#include "Framework/Commands/InputChord.h"
+#include "Framework/MultiBox/MultiBoxDefs.h"
 #include "HAL/Platform.h"
-#include "Internationalization/Text.h"
 #include "Templates/UnrealTemplate.h"
 #include "Templates/SharedPointer.h"
-#include "Textures/SlateIcon.h"
 #include "UObject/GCObject.h"
 
 #include <type_traits>
@@ -44,7 +42,7 @@ public:
 	virtual ~FPromotedFilterContainer() override;
 
 	/** Appends the filters to InToolbarBuilder. */
-	void AppendToBuilder(FToolBarBuilder& InToolBarBuilder) const;
+	void AppendToBuilder(FToolBarBuilder& InToolBarBuilder, const FMenuEntryResizeParams& InResizeParams = {}) const;
 	/** Appends the filters to InMenuBuilder. */
 	void AppendToBuilder(FMenuBuilder& InMenuBuilder) const;
 	
diff --git a/Engine/Source/Editor/CurveEditor/Public/Misc/ResizeParamUtils.h b/Engine/Source/Editor/CurveEditor/Public/Misc/ResizeParamUtils.h
new file mode 100644
--- /dev/null
+++ b/Engine/Source/Editor/CurveEditor/Public/Misc/ResizeParamUtils.h
@@ -0,0 +1,45 @@
+﻿// Copyright Epic Games, Inc. All Rights Reserved.
+
+#pragma once
+
+#include "Containers/StringFwd.h"
+#include "Containers/StringView.h"
+#include "Framework/MultiBox/MultiBoxDefs.h"
+#include "HAL/Platform.h"
+#include "Styling/AppStyle.h"
+
+namespace UE::CurveEditor
+{
+/**
+ * Looks up the CurveEditor.AllowClipping.[UISpecifier] and CurveEditor.ClipPriority.[UISpecifier] from the given style.
+ * @return Resizing params to use for the toolbar item.
+ */
+FMenuEntryResizeParams MakeResizeParams(FStringView UISpecifier, const ISlateStyle& Style = FAppStyle::Get());
+}
+
+namespace UE::CurveEditor
+{
+inline FMenuEntryResizeParams MakeResizeParams(FStringView UISpecifier, const ISlateStyle& Style)
+{
+	// The point of this variable is to check whether UISpecifier exists in the style. We assume that nobody will ever use this ridiculous value...
+	constexpr float UnsetValue = -4200000.f; 
+	
+	const FString AllowClippingString = FString::Printf(TEXT("CurveEditor.AllowClipping.%s"), UISpecifier.GetData());
+	const float AllowClipping_Style = Style.GetFloat(*AllowClippingString, nullptr, UnsetValue);
+	const bool bAllowClipping = AllowClipping_Style == UnsetValue
+		? FMenuEntryResizeParams::DefaultAllowClipping
+		: AllowClipping_Style >= 1.f; // There's no bools in styles, so we'll treat non-zero as true.
+
+	const FString ClipPriorityString = FString::Printf(TEXT("CurveEditor.ClipPriority.%s"), UISpecifier.GetData());
+	const float ClipPriority_Style =  Style.GetFloat(*ClipPriorityString, nullptr, UnsetValue);
+	const float ClipPriority = ClipPriority_Style == UnsetValue ? FMenuEntryResizeParams::DefaultClippingPriority : ClipPriority_Style;
+	
+	// If any of these trigger, either you've misspelled the identifier or you didn't register it.
+	// Add Set("CurveEditor.AllowClipping.[YourSpecifier]", YourPriority) to FStartshipEditorStyle.
+	ensureAlwaysMsgf(AllowClipping_Style != UnsetValue, TEXT("Specifier %s is not mapped in the style!"), UISpecifier.GetData());
+	// Add Set("CurveEditor.ClipPriority.[YourSpecifier]", YourPriority) to FStartshipEditorStyle.
+	ensureAlwaysMsgf(ClipPriority_Style != UnsetValue, TEXT("Specifier %s is not mapped in the style!"), UISpecifier.GetData());
+	
+	return { .ClippingPriority = ClipPriority, .AllowClipping =  bAllowClipping };
+}
+}
\ No newline at end of file
diff --git a/Engine/Source/Editor/CurveEditor/Public/SCurveEditorPanel.h b/Engine/Source/Editor/CurveEditor/Public/SCurveEditorPanel.h
--- a/Engine/Source/Editor/CurveEditor/Public/SCurveEditorPanel.h
+++ b/Engine/Source/Editor/CurveEditor/Public/SCurveEditorPanel.h
@@ -32,6 +32,7 @@
 class FCurveEditor;
 class FExtender;
 class FTabManager;
+class FToolBarBuilder;
 class FMenuBuilder;
 class FUICommandList;
 class IDetailsView;
@@ -200,6 +201,11 @@ private:
 	virtual void OnMouseEnter(const FGeometry& MyGeometry, const FPointerEvent& MouseEvent) override;
     virtual void OnMouseLeave(const FPointerEvent& MouseEvent) override;
 
+	/** @return Extender combining all ICurveEditorModule::GetAllToolBarMenuExtenders and the owning FCurveEditor's ICurveEditorExtensions. */
+	TSharedPtr<FExtender> CombineEditorExtensions() const;
+	/** Adds all base curve editor UI elements to ToolbarBuilder. */
+	void BuildToolbar(FToolBarBuilder& InToolBarBuilder, TSharedPtr<FExtender> InBaseExtender);
+
 	TSharedRef<SWidget> MakeTimeSnapMenu();
 	FText GetTimeSnapMenuTooltip() const;
 
diff --git a/Engine/Source/Editor/EditorStyle/Private/StarshipStyle.cpp b/Engine/Source/Editor/EditorStyle/Private/StarshipStyle.cpp
--- a/Engine/Source/Editor/EditorStyle/Private/StarshipStyle.cpp
+++ b/Engine/Source/Editor/EditorStyle/Private/StarshipStyle.cpp
@@ -7652,7 +7652,7 @@ void FStarshipEditorStyle::FStyle::SetupToolkitStyles()
 		// Some min sizes for toolbar combo buttons. We don't want the combo buttons to resize depending on the size of their contained label.
 		CurveEditorToolbar.SetComboContentHorizontalAlignment(HAlign_Center);
 		{
-			constexpr float LongestTweeName = 90.f;
+			constexpr float LongestTweenName = 90.f;
 			
 			CurveEditorToolbar.SetComboLabelMinWidth(68.f);
 			Set("CurveEditorToolBar.ToolsCombo", CurveEditorToolbar);
@@ -7664,16 +7664,58 @@ void FStarshipEditorStyle::FStyle::SetupToolkitStyles()
 			// Toolbar sliders should appear just a bright as when they are hovered.
 			TweenToolbar.ButtonStyle.SetNormalForeground(TweenColor);
 			TweenToolbar.ComboButtonStyle.ButtonStyle.SetNormalForeground(TweenColor);
-			CurveEditorToolbar.SetComboLabelMinWidth(LongestTweeName);
+			CurveEditorToolbar.SetComboLabelMinWidth(LongestTweenName);
 			Set("CurveEditorTweenToolbar", CurveEditorToolbar);
 			
-			CurveEditorToolbar.SetComboLabelMinWidth(LongestTweeName);
+			CurveEditorToolbar.SetComboLabelMinWidth(LongestTweenName);
 			CurveEditorToolbar.SetAllowWrapButton(false);
 			CurveEditorToolbar.SetBackground(*GetBrush("EditorViewport.OverlayBrush"));
 			Set("ControlRigTweenToolbar", CurveEditorToolbar);
 		}
 
-		Set("CurveEditor.KeyDetailWidth", 180.f);
+		Set("CurveEditor.KeyDetailWidth", 130.f);
+
+		// Toolbar items that are allowed to be clipped. 0.f -> false, >=1.f -> true.
+		Set("CurveEditor.AllowClipping.Sequencer.Save",						0.0f);
+		Set("CurveEditor.AllowClipping.View",								0.0f);
+		Set("CurveEditor.AllowClipping.Tools",								0.0f);
+		Set("CurveEditor.AllowClipping.KeyDetails",							1.0f);
+		Set("CurveEditor.AllowClipping.AxisSnapping",						1.0f);
+		Set("CurveEditor.AllowClipping.ToggleInputSnapping",				1.0f);
+		Set("CurveEditor.AllowClipping.TimeSnapping",						1.0f);
+		Set("CurveEditor.AllowClipping.ToggleOutputSnapping",				1.0f);
+		Set("CurveEditor.AllowClipping.GridSnapping",						1.0f);
+		Set("CurveEditor.AllowClipping.TangentMenu",						1.0f);
+		Set("CurveEditor.AllowClipping.InterpolationToggleWeighted",		1.0f);
+		Set("CurveEditor.AllowClipping.FlattenTangents",					1.0f);
+		Set("CurveEditor.AllowClipping.StraightenTangents",					1.0f);
+		Set("CurveEditor.AllowClipping.CurvesMenu",							1.0f);
+		Set("CurveEditor.AllowClipping.PromotedFilters",					1.0f);
+		Set("CurveEditor.AllowClipping.OpenUserImplementableFilterWindow",	1.0f);
+		Set("CurveEditor.AllowClipping.Tween.FunctionSelect",				1.0f);
+		Set("CurveEditor.AllowClipping.Tween.Slider",						1.0f);
+		Set("CurveEditor.AllowClipping.Tween.Overshoot",					1.0f);
+
+		// Clipping priorities of toolbar items. High number = higher priority -> clipped last.
+		Set("CurveEditor.ClipPriority.Sequencer.Save",						100.0f);
+		Set("CurveEditor.ClipPriority.View",								100.0f);
+		Set("CurveEditor.ClipPriority.Tools",								100.0f);
+		Set("CurveEditor.ClipPriority.KeyDetails",							50.0f);
+		Set("CurveEditor.ClipPriority.AxisSnapping",						40.0f);
+		Set("CurveEditor.ClipPriority.ToggleInputSnapping",					41.0f);
+		Set("CurveEditor.ClipPriority.TimeSnapping",						41.0f);
+		Set("CurveEditor.ClipPriority.ToggleOutputSnapping",				42.0f);
+		Set("CurveEditor.ClipPriority.GridSnapping",						42.0f);
+		Set("CurveEditor.ClipPriority.TangentMenu",							30.0f);
+		Set("CurveEditor.ClipPriority.InterpolationToggleWeighted",			21.0f);
+		Set("CurveEditor.ClipPriority.FlattenTangents",						22.0f);
+		Set("CurveEditor.ClipPriority.StraightenTangents",					23.0f);
+		Set("CurveEditor.ClipPriority.CurvesMenu",							10.0f);
+		Set("CurveEditor.ClipPriority.PromotedFilters",						11.0f);
+		Set("CurveEditor.ClipPriority.OpenUserImplementableFilterWindow",	12.0f);
+		Set("CurveEditor.ClipPriority.Tween.FunctionSelect",				71.0f);
+		Set("CurveEditor.ClipPriority.Tween.Slider",						72.0f);
+		Set("CurveEditor.ClipPriority.Tween.Overshoot",						70.0f);
 
 		// Tab
 		Set("GenericCurveEditor", new IMAGE_BRUSH_SVG("Starship/GenericCurveEditor/CurveEditor", Icon20x20));
diff --git a/Engine/Source/Editor/Sequencer/Private/Widgets/CurveEditor/SSequencerCurveEditor.cpp b/Engine/Source/Editor/Sequencer/Private/Widgets/CurveEditor/SSequencerCurveEditor.cpp
--- a/Engine/Source/Editor/Sequencer/Private/Widgets/CurveEditor/SSequencerCurveEditor.cpp
+++ b/Engine/Source/Editor/Sequencer/Private/Widgets/CurveEditor/SSequencerCurveEditor.cpp
@@ -6,6 +6,7 @@
 #include "MovieSceneCommonHelpers.h"
 #include "SCurveEditorPanel.h"
 #include "Sequencer.h"
+#include "Misc/ResizeParamUtils.h"
 
 #define LOCTEXT_NAMESPACE "SSequencerCurveEditor"
 
@@ -60,7 +61,8 @@ TSharedRef<SWidget> SSequencerCurveEditor::MakeToolbar(TSharedRef<SCurveEditorPa
 				NAME_None,
 				LOCTEXT("SaveDirtyPackages", "Save"),
 				LOCTEXT("SaveDirtyPackagesTooltip", "Saves the current sequence and any subsequences"),
-				SaveIcon
+				SaveIcon,
+				EUserInterfaceActionType::Button, NAME_None, {}, {}, UE::CurveEditor::MakeResizeParams(TEXT("Sequencer.Save"))
 			);
 		}		
 	}
