Curve Editor: Add back Curves section in toolbar, which was accidentally removed in merge of previous CL.

diff --git a/Engine/Source/Editor/CurveEditor/Private/SCurveEditorPanel.cpp b/Engine/Source/Editor/CurveEditor/Private/SCurveEditorPanel.cpp
--- a/Engine/Source/Editor/CurveEditor/Private/SCurveEditorPanel.cpp
+++ b/Engine/Source/Editor/CurveEditor/Private/SCurveEditorPanel.cpp
@@ -1569,11 +1569,28 @@ void SCurveEditorPanel::BuildToolbar(FToolBarBuilder& InToolBarBuilder, TSharedP
 				FOnGetContent::CreateSP(this, &SCurveEditorPanel::MakeFlipCurveMenu, FCurveEditor::ECurveFlipDirection::Vertical),
 				LOCTEXT("FlipCurveVerticalOptions", "Flip Curve Vertical"),
 				LOCTEXT("FlipCurveVerticalOptionsToolTip", "Choose the range for vertical curve flip"),
-				TAttribute<FSlateIcon>(),
+				TAttribute<FSlateIcon>(), 
 				true);
 		}
 		InToolBarBuilder.EndSection();
 	}
+
+	InToolBarBuilder.BeginSection("Curves");
+	{
+		InToolBarBuilder.SetLabelVisibility(EVisibility::Visible); // Show label...
+		InToolBarBuilder.AddComboButton(
+			FUIAction(),
+			FOnGetContent::CreateSP(this, &SCurveEditorPanel::MakeCurvesMenu, InBaseExtender),
+			LOCTEXT("CurveEditor.CurvesCombo.Label", "Curves"),
+			LOCTEXT("CurveEditor.CurvesCombo.ToolTip", "Curves"),
+			FSlateIcon(FAppStyle::Get().GetStyleSetName(), "GenericCurveEditor.Curves")
+			);
+		InToolBarBuilder.SetLabelVisibility(EVisibility::Collapsed); // ... but don't affect the other entries
+				
+		GetCurveEditor()->GetToolbarPromotedFilters()->AppendToBuilder(InToolBarBuilder);
+		InToolBarBuilder.AddToolBarButton(FCurveEditorCommands::Get().OpenUserImplementableFilterWindow);
+	}
+	InToolBarBuilder.EndStyleOverride();
 }
 
 TSharedRef<SWidget> SCurveEditorPanel::MakeTimeSnapMenu()
