Curve Editor: Increased zoom speed in Sequencer's curve editors.

Alt+RMB and Alt+Shift+RMB now zoom with significantly less mouse movement required (often only 1 mouse stroke is required to fully zoom in or out), i.e. "faster".
Zoom out is limited to a maximum of 15 minutes (2700 frames) and 10000 in output axis because it is now very easy to overshoot reasonable view bounds.

For now, this change only affects Sequencer's curve editors. Features, like the Curve Asset Editor, must actively opt-in for this and will continue to zoom just as slowly like before.
To opt-in, specify the FCurveEditorInitParams::ZoomScalingAttr argument when initializing the FCurveEditor.

diff --git a/Engine/Source/Editor/CurveEditor/Private/CurveEditor.cpp b/Engine/Source/Editor/CurveEditor/Private/CurveEditor.cpp
--- a/Engine/Source/Editor/CurveEditor/Private/CurveEditor.cpp
+++ b/Engine/Source/Editor/CurveEditor/Private/CurveEditor.cpp
@@ -10,6 +10,7 @@
 #include "CurveEditorSettings.h"
 #include "CurveEditorSnapMetrics.h"
 #include "CurveEditorAxis.h"
+#include "CurveEditorZoomScaleConfig.h"
 #include "CurveModel.h"
 #include "Curves/KeyHandle.h"
 #include "Curves/RichCurve.h"
@@ -107,6 +108,7 @@ void FCurveEditor::InitCurveEditor(const FCurveEditorInitParams& InInitParams)
 	ICurveEditorModule& CurveEditorModule = FModuleManager::LoadModuleChecked<ICurveEditorModule>("CurveEditor");
 
 	Selection = FCurveEditorSelection(SharedThis(this));
+	ZoomScalingAttr = InInitParams.ZoomScalingAttr;
 
 	// Editor Extensions can be registered in the Curve Editor module. To allow users to derive from FCurveEditor
 	// we have to manually reach out to the module and get a list of extensions to create an instance of them.
@@ -123,6 +125,7 @@ void FCurveEditor::InitCurveEditor(const FCurveEditorInitParams& InInitParams)
 		EditorExtensions.Add(NewExtension);
 	}
 
+
 	TArrayView<const FOnCreateCurveEditorToolExtension> Tools = CurveEditorModule.GetToolExtensions();
 	for (int32 DelegateIndex = 0; DelegateIndex < Tools.Num(); ++DelegateIndex)
 	{
@@ -666,6 +669,15 @@ void FCurveEditor::ZoomToFitSelection(EAxisList::Type Axes)
 	ZoomToFitInternal(Axes, Selection.GetAll());
 }
 
+const FCurveEditorZoomScaleConfig& FCurveEditor::GetZoomScaleConfig() const
+{
+	const bool bCanCallGet = ZoomScalingAttr.IsBound() || ZoomScalingAttr.IsSet();
+	const FCurveEditorZoomScaleConfig* OverrideConfig = bCanCallGet ? ZoomScalingAttr.Get() : nullptr;
+	
+	static FCurveEditorZoomScaleConfig Default;
+	return OverrideConfig ? *OverrideConfig : Default;
+}
+
 void FCurveEditor::ZoomToFitInternal(EAxisList::Type Axes, const TMap<FCurveModelID, FKeyHandleSet>& CurveKeySet)
 {
 	TArray<FKeyPosition> KeyPositionsScratch;
diff --git a/Engine/Source/Editor/CurveEditor/Private/DragOperations/CurveEditorDragOperation_Pan.cpp b/Engine/Source/Editor/CurveEditor/Private/DragOperations/CurveEditorDragOperation_Pan.cpp
--- a/Engine/Source/Editor/CurveEditor/Private/DragOperations/CurveEditorDragOperation_Pan.cpp
+++ b/Engine/Source/Editor/CurveEditor/Private/DragOperations/CurveEditorDragOperation_Pan.cpp
@@ -4,6 +4,7 @@
 
 #include "CurveEditor.h"
 #include "CurveEditorScreenSpace.h"
+#include "CurveEditorZoomScaleConfig.h"
 #include "ICurveEditorBounds.h"
 #include "Input/Events.h"
 #include "Math/UnrealMathUtility.h"
@@ -52,11 +53,13 @@ FReply FCurveEditorDragOperation_PanView::OnMouseWheel(FVector2D InitialPosition
 	if (bIsDragging)
 	{
 		FCurveEditorScreenSpace ViewSpace = View->GetViewSpace();
-		float ZoomDelta = 1.f - FMath::Clamp(0.1f * MouseEvent.GetWheelDelta(), -0.9f, 0.9f);
 
 		double CurrentTime = ViewSpace.ScreenToSeconds(CurrentPosition.X);
 		double CurrentValue = ViewSpace.ScreenToValue(CurrentPosition.Y);
 
+
+		const double WheelMultiplier = CurveEditor->GetZoomScaleConfig().GetMouseWheelZoomMultiplierClamped();
+		const double ZoomDelta = 1.0 - FMath::Clamp(0.1 * WheelMultiplier * MouseEvent.GetWheelDelta(), -0.9, 0.9);
 		View->ZoomAround(FVector2D(ZoomDelta, ZoomDelta), CurrentTime, CurrentValue);
 
 		// Adjust the stored initial bounds by the zoom delta so delta calculations work properly
diff --git a/Engine/Source/Editor/CurveEditor/Private/DragOperations/CurveEditorDragOperation_Zoom.cpp b/Engine/Source/Editor/CurveEditor/Private/DragOperations/CurveEditorDragOperation_Zoom.cpp
--- a/Engine/Source/Editor/CurveEditor/Private/DragOperations/CurveEditorDragOperation_Zoom.cpp
+++ b/Engine/Source/Editor/CurveEditor/Private/DragOperations/CurveEditorDragOperation_Zoom.cpp
@@ -4,14 +4,27 @@
 
 #include "CurveEditor.h"
 #include "CurveEditorScreenSpace.h"
+#include "CurveEditorZoomScaleConfig.h"
 #include "ICurveEditorBounds.h"
+#include "SCurveEditor.h"
 #include "Input/Events.h"
 #include "Math/UnrealMathUtility.h"
 #include "SCurveEditorView.h"
 
+TAutoConsoleVariable<bool> CVarLogVerticalZoomMultipliers(
+	TEXT("CurveEditor.LogVerticalZoomMultipliers"),
+	true,
+	TEXT("Logs the zoom multipliers to make it easier for you to tweak FCurveEditorZoomScaleConfig::VerticalZoomScale")
+	);
+TAutoConsoleVariable<bool> CVarLogHorizontalZoomMultipliers(
+	TEXT("CurveEditor.LogHorizontalZoomMultipliers"),
+	true,
+	TEXT("Logs the zoom multipliers to make it easier for you to tweak FCurveEditorZoomScaleConfig::HorizontalZoomScale")
+	);
+
 FCurveEditorDragOperation_Zoom::FCurveEditorDragOperation_Zoom(FCurveEditor* InCurveEditor, TSharedPtr<SCurveEditorView> InOptionalView)
 	: CurveEditor(InCurveEditor)
-	, OptionalView(InOptionalView)
+	, OptionalView(MoveTemp(InOptionalView))
 {}
 
 void FCurveEditorDragOperation_Zoom::OnBeginDrag(FVector2D InitialPosition, FVector2D CurrentPosition, const FPointerEvent& MouseEvent)
@@ -34,28 +47,34 @@ void FCurveEditorDragOperation_Zoom::OnBeginDrag(FVector2D InitialPosition, FVec
 
 void FCurveEditorDragOperation_Zoom::OnDrag(FVector2D InitialPosition, FVector2D CurrentPosition, const FPointerEvent& MouseEvent)
 {
-	FVector2D PixelDelta = CurrentPosition - InitialPosition;
-
-	double ClampRange = 1e9f;
+	constexpr double ClampRange = 1e9f;
+	const FVector2D PixelDelta = CurrentPosition - InitialPosition;
+	const FCurveEditorZoomScaleConfig& ZoomConfig = CurveEditor->GetZoomScaleConfig();
 
 	// Zoom input range
-	FCurveEditorScreenSpaceH InputSpace  = OptionalView ? OptionalView->GetViewSpace() : CurveEditor->GetPanelInputSpace();
+	const FCurveEditorScreenSpaceH InputSpace  = OptionalView
+		? static_cast<FCurveEditorScreenSpaceH>(OptionalView->GetViewSpace()) : CurveEditor->GetPanelInputSpace();
 
-	double DiffX = double(PixelDelta.X) / (InputSpace.GetPhysicalWidth() / OriginalInputRange);
+	const double InputZoomMultiplier = GetZoomMultiplier_InputAxis(-PixelDelta.X);
+	const double DiffX = InputZoomMultiplier * PixelDelta.X / (InputSpace.GetPhysicalWidth() / OriginalInputRange);
 
 	// This flips the horizontal zoom to match existing DCC tools
-	double NewInputRange = OriginalInputRange - DiffX;
-	double InputMin = FMath::Clamp<double>(ZoomOriginX - NewInputRange * ZoomFactor.X, -ClampRange, ClampRange);
-	double InputMax = FMath::Clamp<double>(ZoomOriginX + NewInputRange * (1.f - ZoomFactor.X), InputMin, ClampRange);
-
-	CurveEditor->GetBounds().SetInputBounds(InputMin, InputMax);
+	const double NewInputRange = OriginalInputRange - DiffX;
+	const double InputMin = FMath::Clamp<double>(ZoomOriginX - NewInputRange * ZoomFactor.X, -ClampRange, ClampRange);
+	const double InputMax = FMath::Clamp<double>(ZoomOriginX + NewInputRange * (1.f - ZoomFactor.X), InputMin, ClampRange);
+	const bool bExceedsHorizontalLimit = ZoomConfig.bLimitHorizontalZoomOut && ZoomConfig.MaxHorizontalZoomOut <= InputMax - InputMin;
+	if (!bExceedsHorizontalLimit)
+	{
+		CurveEditor->GetBounds().SetInputBounds(InputMin, InputMax);
+	}
 
 	// Zoom output range
 	if (OptionalView.IsValid())
 	{
 		FCurveEditorScreenSpaceV ViewSpace = OptionalView->GetViewSpace();
 
-		double DiffY = double(PixelDelta.Y) / (ViewSpace.GetPhysicalHeight() / OriginalOutputRange);
+		const double OutputZoomMultiplier = GetZoomMultiplier_OutputAxis(PixelDelta.Y);
+		const double DiffY = OutputZoomMultiplier * PixelDelta.Y / (ViewSpace.GetPhysicalHeight() / OriginalOutputRange);
 		double NewOutputRange = OriginalOutputRange + DiffY;
 
 		// If they're holding a shift they can scale on both axis at once, non-proportionally
@@ -65,9 +84,29 @@ void FCurveEditorDragOperation_Zoom::OnDrag(FVector2D InitialPosition, FVector2D
 			NewOutputRange = (NewInputRange / OriginalInputRange) * OriginalOutputRange;
 		}
 
-		double OutputMin = FMath::Clamp<double>(ZoomOriginY - NewOutputRange * (1.f - ZoomFactor.Y), -ClampRange, ClampRange);
-		double OutputMax = FMath::Clamp<double>(ZoomOriginY + NewOutputRange * (ZoomFactor.Y), OutputMin, ClampRange);
+		const double OutputMin = FMath::Clamp<double>(ZoomOriginY - NewOutputRange * (1.f - ZoomFactor.Y), -ClampRange, ClampRange);
+		const double OutputMax = FMath::Clamp<double>(ZoomOriginY + NewOutputRange * (ZoomFactor.Y), OutputMin, ClampRange);
 
-		OptionalView->SetOutputBounds(OutputMin, OutputMax);
+		const bool bExceedsVerticalLimit = ZoomConfig.bLimitVerticalZoomOut && ZoomConfig.MaxVerticalZoomOut <= OutputMax - OutputMin;
+		if (!bExceedsVerticalLimit)
+		{
+			OptionalView->SetOutputBounds(OutputMin, OutputMax);
+		}
 	}
 }
+
+double FCurveEditorDragOperation_Zoom::GetZoomMultiplier_InputAxis(double InMovedMouseX) const
+{
+	const FCurveEditorZoomScaleConfig& ZoomConfig = CurveEditor->GetZoomScaleConfig();
+	const double Factor = ZoomConfig.EvalHorizontalZoom(InMovedMouseX);
+	UE_CLOG(CVarLogHorizontalZoomMultipliers.GetValueOnAnyThread(), LogCurveEditor, Log, TEXT("Horizontal Zoom: X: %f, Factor: %f"), InMovedMouseX, Factor);
+	return Factor;
+}
+
+double FCurveEditorDragOperation_Zoom::GetZoomMultiplier_OutputAxis(double InMovedMouseY) const
+{
+	const FCurveEditorZoomScaleConfig& ZoomConfig = CurveEditor->GetZoomScaleConfig();
+	const double Factor = ZoomConfig.EvalVerticalZoom(InMovedMouseY);
+	UE_CLOG(CVarLogVerticalZoomMultipliers.GetValueOnAnyThread(), LogCurveEditor, Log, TEXT("Vertical Zoom: X: %f, Factor: %f"), InMovedMouseY, Factor);
+	return Factor;
+}
\ No newline at end of file
diff --git a/Engine/Source/Editor/CurveEditor/Private/DragOperations/CurveEditorDragOperation_Zoom.h b/Engine/Source/Editor/CurveEditor/Private/DragOperations/CurveEditorDragOperation_Zoom.h
--- a/Engine/Source/Editor/CurveEditor/Private/DragOperations/CurveEditorDragOperation_Zoom.h
+++ b/Engine/Source/Editor/CurveEditor/Private/DragOperations/CurveEditorDragOperation_Zoom.h
@@ -4,6 +4,7 @@
 
 #include "ICurveEditorDragOperation.h"
 #include "Math/Vector2D.h"
+#include "Misc/Attribute.h"
 #include "Templates/SharedPointer.h"
 
 class FCurveEditor;
@@ -13,7 +14,8 @@ struct FPointerEvent;
 class FCurveEditorDragOperation_Zoom : public ICurveEditorDragOperation
 {
 public:
-	FCurveEditorDragOperation_Zoom(FCurveEditor* CurveEditor, TSharedPtr<SCurveEditorView> InOptionalView);
+	
+	FCurveEditorDragOperation_Zoom(FCurveEditor* InCurveEditor, TSharedPtr<SCurveEditorView> InOptionalView = nullptr);
 
 	virtual void OnBeginDrag(FVector2D InitialPosition, FVector2D CurrentPosition, const FPointerEvent& MouseEvent) override;
 	virtual void OnDrag(FVector2D InitialPosition, FVector2D CurrentPosition, const FPointerEvent& MouseEvent) override;
@@ -26,4 +28,7 @@ private:
 	FVector2D ZoomFactor;
 	double ZoomOriginX, ZoomOriginY;
 	double OriginalInputRange, OriginalOutputRange;
-};
\ No newline at end of file
+
+	double GetZoomMultiplier_InputAxis(double InMovedMouseX) const;
+	double GetZoomMultiplier_OutputAxis(double InMovedMouseY) const;
+};
diff --git a/Engine/Source/Editor/CurveEditor/Private/Views/SInteractiveCurveEditorView.cpp b/Engine/Source/Editor/CurveEditor/Private/Views/SInteractiveCurveEditorView.cpp
--- a/Engine/Source/Editor/CurveEditor/Private/Views/SInteractiveCurveEditorView.cpp
+++ b/Engine/Source/Editor/CurveEditor/Private/Views/SInteractiveCurveEditorView.cpp
@@ -23,6 +23,7 @@
 #include "CurveEditorSettings.h"
 #include "CurveEditorSnapMetrics.h"
 #include "CurveEditorTypes.h"
+#include "CurveEditorZoomScaleConfig.h"
 #include "CurveModel.h"
 #include "Curves/KeyHandle.h"
 #include "Curves/RichCurve.h"
@@ -1396,7 +1397,8 @@ FReply SInteractiveCurveEditorView::OnMouseWheel(const FGeometry& MyGeometry, co
 		}
 	}
 
-	float ZoomDelta = 1.f - FMath::Clamp(0.1f * MouseEvent.GetWheelDelta(), -0.9f, 0.9f);
+	const double WheelMultiplier = CurveEditor->GetZoomScaleConfig().GetMouseWheelZoomMultiplierClamped();
+	const double ZoomDelta = 1.0 - FMath::Clamp(0.1 * WheelMultiplier * MouseEvent.GetWheelDelta(), -0.9, 0.9);
 	ZoomAround(FVector2D(ZoomDelta, ZoomDelta), CurrentTime, CurrentValue);
 
 	return FReply::Handled();
diff --git a/Engine/Source/Editor/CurveEditor/Public/CurveEditor.h b/Engine/Source/Editor/CurveEditor/Public/CurveEditor.h
--- a/Engine/Source/Editor/CurveEditor/Public/CurveEditor.h
+++ b/Engine/Source/Editor/CurveEditor/Public/CurveEditor.h
@@ -180,6 +180,10 @@ public:
 	 * @param Axes         (Optional) Axes to lock the zoom to
 	 */
 	void ZoomToFitSelection(EAxisList::Type Axes = EAxisList::All);
+
+	/** @return The config to use for scaling the zooms. */
+	const FCurveEditorZoomScaleConfig& GetZoomScaleConfig() const;
+	
 	/**
 	* Assign a new bounds container to this curve editor
 	*/
@@ -731,6 +735,9 @@ protected:
 	/** List of editor extensions we have initialized. */
 	TArray<TSharedRef<ICurveEditorExtension>> EditorExtensions;
 
+	/** The zoom scale config that was passed in to us in InitCurveEditor. */
+	TAttribute<const FCurveEditorZoomScaleConfig*> ZoomScalingAttr;
+
 	/** List of tool extensions we have initialized. */
 	TMap<FCurveEditorToolID, TUniquePtr<ICurveEditorToolExtension>> ToolExtensions;
 
diff --git a/Engine/Source/Editor/CurveEditor/Public/CurveEditorZoomScaleConfig.h b/Engine/Source/Editor/CurveEditor/Public/CurveEditorZoomScaleConfig.h
new file mode 100644
--- /dev/null
+++ b/Engine/Source/Editor/CurveEditor/Public/CurveEditorZoomScaleConfig.h
@@ -0,0 +1,117 @@
+﻿// Copyright Epic Games, Inc. All Rights Reserved.
+
+#pragma once
+
+#include "Curves/CurveFloat.h"
+#include "CurveEditorZoomScaleConfig.generated.h"
+
+USTRUCT()
+struct FCurveEditorZoomScaleConfig
+{
+	GENERATED_BODY()
+
+	/** Multiplier when zooming uniformly using mouse wheel. 1.0 means normal. Values < 1.0, make zoom 'slower', and > 1.0 zoom 'faster'. */
+	UPROPERTY(config, EditAnywhere, Category="Curve Editor", meta = (ClampMin = 0.05, ClampMax = 32.0))
+	double MouseWheelZoomMultiplier = 1.0;
+
+	/**
+	 * Easing function when zooming out on the X-axis using Alt+Shift+RMB + move mouse left/right.
+	 *
+	 * X-value is how far the mouse has travelled on the X-axis. Positive when zooming out, negative when zooming in.
+	 * Y-Value is the resulting zoom multiplier. Must be >= 0.0. When zooming out (i.e. x < 0), the Y value must be <= 1.0 or zooming is unstable.
+	 * Tip: You can use "CurveEditor.LogHorizontalZoomMultipliers" CVar to make tuning the x and y values easier.
+	 *
+	 * Example:
+	 * - For Sequencer, the frame rate is the length of 1 (unit length).
+	 * - So if Sequencer's target FPS is set to 30, and 45 frames are visible on the X-Axis in Curve Editor, we say the X-axis is 1.5 units wide.
+	 * - We'd look up the value 1.5 in this easing function and use that as zoom multiplier.
+	 */
+	UPROPERTY(config, EditAnywhere, Category="Curve Editor", meta=(
+		DisplayName="Right mouse (input axis)",
+		XAxisName="Input Axis Length",
+		YAxisName="Zoom Factor"))
+	FRuntimeFloatCurve HorizontalZoomScale;
+	
+	/**
+	 * Easing function when zooming out on the Y-axis using Alt+RMB + move mouse up/down.
+	 *
+	 * X-value is how far the mouse has travelled on the Y-axis.  Positive when zooming out, negative when zooming in.
+	 * Y-Value is the resulting zoom multiplier. Must be >= 0.0. When zooming out (i.e. x < 0), the Y value must be <= 1.0 or zooming is unstable.
+	 * Tip: You can use "CurveEditor.LogVerticalZoomMultipliers" CVar to make tuning the x and y values easier.
+	 *
+	 * Example:
+	 * - Suppose in Curve Editor, the Y-axis bounds are +2000 to -1000.
+	 * - The height of the Y-axis is thus 3000.
+	 * - We'd look up the value of 3000 in this easing function and use that as zoom multiplier.
+	 */
+	UPROPERTY(config, EditAnywhere, Category="Curve Editor", meta=(
+		DisplayName="Right mouse (output axis)",
+		XAxisName="Output Axis Length",
+		YAxisName="Zoom Factor"))
+	FRuntimeFloatCurve VerticalZoomScale;
+
+	/** If true, when zooming out using right-click the vertical axis size will be limited. */
+	UPROPERTY(config, EditAnywhere, Category="Curve Editor", meta = (ClampMin = 0.00))
+	bool bLimitHorizontalZoomOut = false;
+	/**
+	 * If bLimitHorizontalZoomOut is true, then this max allowed value range the input axis can show.
+	 * Example: If this is 1000, and you have zoomed out so the axis shows -200 and 800, you can zoom out no further.
+	 * 
+	 * For Sequencer, this is the max number of seconds. So, you need multiply this with the target FPS. So 1000 would mean a max of 3000 frames
+	 * with a target FPS of 30. 
+	 */
+	UPROPERTY(config, EditAnywhere, Category="Curve Editor", meta = (ClampMin = 0.00, EditCondition = "!bUnlimitedHorizontalZoomOut", EditConditionHides))
+	double MaxHorizontalZoomOut = 10000.0;
+	
+	/** If true, when zooming out using right-click the horizontal axis size will be limited. */
+	UPROPERTY(config, EditAnywhere, Category="Curve Editor", meta = (ClampMin = 0.00))
+	bool bLimitVerticalZoomOut = false;
+	/**
+	 * If bLimitVerticalZoomOut is true, then this max allowed value range the output axis can show.
+	 * Example: If this is 1000, and you have zoomed out so the axis shows -200 and 800, you can zoom out no further.
+	 */
+	UPROPERTY(config, EditAnywhere, Category="Curve Editor", meta = (ClampMin = 0.00, EditCondition = "!bUnlimitedVerticalZoomOut", EditConditionHides))
+	double MaxVerticalZoomOut = 10000.0;
+	
+	FCurveEditorZoomScaleConfig()
+	{
+		// By default, set this data up to do not additional scaling of the zoom.
+		HorizontalZoomScale.EditorCurveData.AddKey(0.0, 1.0);
+		HorizontalZoomScale.EditorCurveData.AddKey(1.0, 1.0); // So post infinity linear interpolation works
+		HorizontalZoomScale.EditorCurveData.DefaultValue = 1.0;
+		// While there is a limit for zooming in, for zooming out we want linear extrap so the user does not have to define an infinitely large X.
+		HorizontalZoomScale.EditorCurveData.PostInfinityExtrap = RCCE_Linear;
+		
+		VerticalZoomScale.EditorCurveData.AddKey(0.0, 1.0);
+		VerticalZoomScale.EditorCurveData.AddKey(1.0, 1.0); // So post infinity linear interpolation works
+		VerticalZoomScale.EditorCurveData.DefaultValue = 1.0;
+		// While there is a limit for zooming in, for zooming out we want linear extrap so the user does not have to define an infinitely large X.
+		VerticalZoomScale.EditorCurveData.PostInfinityExtrap = RCCE_Linear;
+	}
+
+	double GetMouseWheelZoomMultiplierClamped() const { return FMath::Max(MouseWheelZoomMultiplier, 0.05); }
+
+	/**
+	 * Evaluates VerticalZoomScale.
+	 * @param InAxisSize Absolute value is Max - Min displayed axis values. Positive when zooming out, negative when zooming in.
+	 * @return Correctly clamped zoom factor.
+	 */
+	double EvalHorizontalZoom(double InAxisSize) const
+	{
+		const double Factor = HorizontalZoomScale.EditorCurveData.Eval(InAxisSize);
+		const bool bIsZoomingIn = Factor < 0.f;
+		return bIsZoomingIn ? FMath::Clamp(Factor, 0.05, 1.f) : FMath::Max(0.05, Factor);
+	}
+
+	/**
+	 * Evaluates VerticalZoomScale.
+	 * @param InAxisSize Absolute value is Max - Min displayed axis values. Positive when zooming out, negative when zooming in.
+	 * @return Correctly clamped zoom factor.
+	 */
+	double EvalVerticalZoom(double InAxisSize) const
+	{
+		const double Factor = VerticalZoomScale.EditorCurveData.Eval(InAxisSize);
+		const bool bIsZoomingIn = Factor < 0.f;
+		return bIsZoomingIn ? FMath::Clamp(Factor, 0.05, 1.f) : FMath::Max(0.05, Factor);
+	}
+};
\ No newline at end of file
diff --git a/Engine/Source/Editor/CurveEditor/Public/ICurveEditorModule.h b/Engine/Source/Editor/CurveEditor/Public/ICurveEditorModule.h
--- a/Engine/Source/Editor/CurveEditor/Public/ICurveEditorModule.h
+++ b/Engine/Source/Editor/CurveEditor/Public/ICurveEditorModule.h
@@ -2,12 +2,13 @@
 
 #pragma  once
 
+#include "Misc/Attribute.h"
 #include "Modules/ModuleManager.h"
 #include "Modules/ModuleInterface.h"
 #include "Containers/ArrayView.h"
 #include "CurveEditorTypes.h"
 
-
+struct FCurveEditorZoomScaleConfig;
 class FCurveEditor;
 class SCurveEditorView;
 class ICurveEditorExtension;
@@ -36,6 +37,9 @@ struct FCurveEditorInitParams
 	
 	/** Extensions you want to inject for this curve editor. */
 	TArray<TSharedRef<ICurveEditorExtension>> AdditionalEditorExtensions;
+
+	/** Optional. If set, defines multipliers for zooming. */
+	TAttribute<const FCurveEditorZoomScaleConfig*> ZoomScalingAttr;
 };
 
 /**
@@ -45,7 +49,7 @@ class ICurveEditorModule : public IModuleInterface
 {
 public:
 
-	static ICurveEditorModule& Get()
+static ICurveEditorModule& Get()
 	{
 		return FModuleManager::Get().GetModuleChecked<ICurveEditorModule>("CurveEditor");
 	}
diff --git a/Engine/Source/Editor/Sequencer/Private/MVVM/CurveEditorExtension.cpp b/Engine/Source/Editor/Sequencer/Private/MVVM/CurveEditorExtension.cpp
--- a/Engine/Source/Editor/Sequencer/Private/MVVM/CurveEditorExtension.cpp
+++ b/Engine/Source/Editor/Sequencer/Private/MVVM/CurveEditorExtension.cpp
@@ -289,11 +289,15 @@ void FCurveEditorExtension::CreateCurveEditor(const FTimeSliderArgs& TimeSliderA
 
 	// Create the curve editor;
 	{
+		USequencerSettings* SequencerSettings = Sequencer->GetSequencerSettings();
+		
 		FCurveEditorInitParams CurveEditorInitParams;
 		CurveEditorInitParams.AdditionalEditorExtensions = { MakeShared<FSequencerCurveEditorToolbarExtender>(Sequencer.ToWeakPtr()) };
-
-		USequencerSettings* SequencerSettings = Sequencer->GetSequencerSettings();
-
+		CurveEditorInitParams.ZoomScalingAttr.BindLambda([SequencerSettings]
+		{
+			return &SequencerSettings->GetCurveEditorZoomScaling();
+		});
+		
 		CurveEditorModel = MakeShared<FSequencerCurveEditor>(Sequencer, TimeSliderArgs.NumericTypeInterface);
 		CurveEditorModel->SetBounds(MakeUnique<FSequencerCurveEditorBounds>(Sequencer.ToSharedRef()));
 		CurveEditorModel->InitCurveEditor(CurveEditorInitParams);
diff --git a/Engine/Source/Editor/Sequencer/Private/SequencerSettings.cpp b/Engine/Source/Editor/Sequencer/Private/SequencerSettings.cpp
--- a/Engine/Source/Editor/Sequencer/Private/SequencerSettings.cpp
+++ b/Engine/Source/Editor/Sequencer/Private/SequencerSettings.cpp
@@ -84,6 +84,15 @@ USequencerSettings::USequencerSettings( const FObjectInitializer& ObjectInitiali
 	SectionColorTints.Add(FColor(147, 119, 101, 255)); // orange
 	SectionColorTints.Add(FColor(139, 95, 108, 255)); // red 
 	SectionColorTints.Add(FColor(109, 74, 121, 255)); // purple
+	
+	// Positive x -> zooming out, Negative x -> zooming in. X is in seconds.
+	// Don't add any factors for zoom-in. Zoom-in is already very fast with a zoom factor of 1.
+	CurveEditorZoomScaling.HorizontalZoomScale.EditorCurveData.AddKey(100, 6); // ... e.g. showing 100 seconds worth of frames, zoom out factor is 6.
+	CurveEditorZoomScaling.VerticalZoomScale.EditorCurveData.AddKey(100, 6);
+	CurveEditorZoomScaling.bLimitHorizontalZoomOut = true;
+	CurveEditorZoomScaling.bLimitVerticalZoomOut = true;
+	CurveEditorZoomScaling.MaxHorizontalZoomOut = 900; // Cannot zoom out more than 900 seconds = 15 minutes.
+	CurveEditorZoomScaling.MaxVerticalZoomOut = 10000;
 }
 
 void USequencerSettings::PostEditChangeProperty(FPropertyChangedEvent& PropertyChangedEvent)
@@ -430,6 +439,12 @@ void USequencerSettings::SetAutoScrollEnabled(bool bInAutoScrollEnabled)
 	}
 }
 
+void USequencerSettings::SetCurveEditorZoomScaling(const FCurveEditorZoomScaleConfig& Value)
+{
+	CurveEditorZoomScaling = Value;
+	SaveConfig();
+}
+
 ESequencerLoopMode USequencerSettings::GetLoopMode() const
 {
 	return LoopMode;
diff --git a/Engine/Source/Editor/Sequencer/Public/SequencerSettings.h b/Engine/Source/Editor/Sequencer/Public/SequencerSettings.h
--- a/Engine/Source/Editor/Sequencer/Public/SequencerSettings.h
+++ b/Engine/Source/Editor/Sequencer/Public/SequencerSettings.h
@@ -3,6 +3,7 @@
 #pragma once
 
 #include "CoreMinimal.h"
+#include "CurveEditorZoomScaleConfig.h"
 #include "Filters/SequencerFilterBarConfig.h"
 #include "Misc/SequencerThumbnailCaptureSettings.h"
 #include "UObject/ObjectMacros.h"
@@ -138,7 +139,7 @@ public:
 	DECLARE_MULTICAST_DELEGATE_OneParam( FOnAllowEditsModeChanged, EAllowEditsMode );
 	DECLARE_MULTICAST_DELEGATE(FOnLoopStateChanged);
 	DECLARE_MULTICAST_DELEGATE(FOnTimeDisplayFormatChanged);
-
+	
 	virtual void PostEditChangeProperty(FPropertyChangedEvent& PropertyChangedEvent) override;
 
 	/** Gets the current auto change mode. */
@@ -289,6 +290,11 @@ public:
 	/** Assign whether we should filter the curve editor tree to only nodes that are relevant to the current sequencer selection */
 	void SetCurveEditorVisible(bool bCurveEditorVisible);
 
+	/** @return The config affecting the zooming behavior of Sequencer's Curve Editor. */
+	const FCurveEditorZoomScaleConfig& GetCurveEditorZoomScaling() { return CurveEditorZoomScaling; }
+	/** Sets the config affecting the zooming behaviour of Sequencer's Curve Editor. */
+	void SetCurveEditorZoomScaling(const FCurveEditorZoomScaleConfig& Value);
+
 	/** Gets the loop mode. */
 	ESequencerLoopMode GetLoopMode() const;
 	/** Sets the loop mode. */
@@ -677,6 +683,10 @@ protected:
 	UPROPERTY(config, EditAnywhere, Category=CurveEditor)
 	bool bCurveEditorVisible;
 
+	/** Affects the zooming behavior of in Sequencer's Curve Editor. */
+	UPROPERTY(config, EditAnywhere, Category=CurveEditor)
+	FCurveEditorZoomScaleConfig CurveEditorZoomScaling;
+
 	/** The loop mode of the playback in timeline. */
 	UPROPERTY( config )
 	TEnumAsByte<ESequencerLoopMode> LoopMode;
