Curve Editor: Fix EditCondition

bUnlimitedHorizontalZoomOut and bUnlimitedVerticalZoomOut don't exist.
Follow up to 41859177

#rnx

diff --git a/Engine/Source/Editor/CurveEditor/Public/CurveEditorZoomScaleConfig.h b/Engine/Source/Editor/CurveEditor/Public/CurveEditorZoomScaleConfig.h
--- a/Engine/Source/Editor/CurveEditor/Public/CurveEditorZoomScaleConfig.h
+++ b/Engine/Source/Editor/CurveEditor/Public/CurveEditorZoomScaleConfig.h
@@ -1,4 +1,4 @@
-﻿// Copyright Epic Games, Inc. All Rights Reserved.
+// Copyright Epic Games, Inc. All Rights Reserved.
 
 #pragma once
 
@@ -60,7 +60,7 @@ struct FCurveEditorZoomScaleConfig
 	 * For Sequencer, this is the max number of seconds. So, you need multiply this with the target FPS. So 1000 would mean a max of 3000 frames
 	 * with a target FPS of 30. 
 	 */
-	UPROPERTY(config, EditAnywhere, Category="Curve Editor", meta = (ClampMin = 0.00, EditCondition = "!bUnlimitedHorizontalZoomOut", EditConditionHides))
+	UPROPERTY(config, EditAnywhere, Category="Curve Editor", meta = (ClampMin = 0.00, EditCondition = "bLimitHorizontalZoomOut", EditConditionHides))
 	double MaxHorizontalZoomOut = 10000.0;
 	
 	/** If true, when zooming out using right-click the horizontal axis size will be limited. */
@@ -70,7 +70,7 @@ struct FCurveEditorZoomScaleConfig
 	 * If bLimitVerticalZoomOut is true, then this max allowed value range the output axis can show.
 	 * Example: If this is 1000, and you have zoomed out so the axis shows -200 and 800, you can zoom out no further.
 	 */
-	UPROPERTY(config, EditAnywhere, Category="Curve Editor", meta = (ClampMin = 0.00, EditCondition = "!bUnlimitedVerticalZoomOut", EditConditionHides))
+	UPROPERTY(config, EditAnywhere, Category="Curve Editor", meta = (ClampMin = 0.00, EditCondition = "bLimitVerticalZoomOut", EditConditionHides))
 	double MaxVerticalZoomOut = 10000.0;
 	
 	FCurveEditorZoomScaleConfig()
