Curve Editor: Remove commands that aren't necessary.

Sequencer commands and bindings for these operations are sufficient and the curve editor doesn't need to have its own bindings for it.

This fixes step to next/previous key not being bound by the playback range.
This also fixes the time slider controller not being set in the correct playback mode (ie. scrubbing) when stepping from key to key to allow audio sounds to play

diff --git a/Engine/Source/Editor/CurveEditor/Private/CurveEditor.cpp b/Engine/Source/Editor/CurveEditor/Private/CurveEditor.cpp
--- a/Engine/Source/Editor/CurveEditor/Private/CurveEditor.cpp
+++ b/Engine/Source/Editor/CurveEditor/Private/CurveEditor.cpp
@@ -476,13 +476,6 @@ void FCurveEditor::BindCommands()
 	CommandList->MapAction(FCurveEditorCommands::Get().TranslateSelectedKeysLeft, FExecuteAction::CreateSP(this, &FCurveEditor::TranslateSelectedKeysLeft));
 	CommandList->MapAction(FCurveEditorCommands::Get().TranslateSelectedKeysRight, FExecuteAction::CreateSP(this, &FCurveEditor::TranslateSelectedKeysRight));
 
-	CommandList->MapAction(FCurveEditorCommands::Get().StepToNextKey, FExecuteAction::CreateSP(this, &FCurveEditor::StepToNextKey));
-	CommandList->MapAction(FCurveEditorCommands::Get().StepToPreviousKey, FExecuteAction::CreateSP(this, &FCurveEditor::StepToPreviousKey));
-	CommandList->MapAction(FCurveEditorCommands::Get().StepForward, FExecuteAction::CreateSP(this, &FCurveEditor::StepForward), EUIActionRepeatMode::RepeatEnabled);
-	CommandList->MapAction(FCurveEditorCommands::Get().StepBackward, FExecuteAction::CreateSP(this, &FCurveEditor::StepBackward), EUIActionRepeatMode::RepeatEnabled);
-	CommandList->MapAction(FCurveEditorCommands::Get().JumpToStart, FExecuteAction::CreateSP(this, &FCurveEditor::JumpToStart));
-	CommandList->MapAction(FCurveEditorCommands::Get().JumpToEnd, FExecuteAction::CreateSP(this, &FCurveEditor::JumpToEnd));
-
 	CommandList->MapAction(FCurveEditorCommands::Get().SetSelectionRangeStart, FExecuteAction::CreateSP(this, &FCurveEditor::SetSelectionRangeStart));
 	CommandList->MapAction(FCurveEditorCommands::Get().SetSelectionRangeEnd, FExecuteAction::CreateSP(this, &FCurveEditor::SetSelectionRangeEnd));
 	CommandList->MapAction(FCurveEditorCommands::Get().ClearSelectionRange, FExecuteAction::CreateSP(this, &FCurveEditor::ClearSelectionRange));
@@ -999,189 +992,6 @@ void FCurveEditor::SnapToSelectedKey()
 	}
 }
 
-void FCurveEditor::StepToNextKey()
-{
-	TSharedPtr<ITimeSliderController> TimeSliderController = WeakTimeSliderController.Pin();
-	if (!TimeSliderController.IsValid())
-	{
-		return;
-	}
-
-	FFrameRate TickResolution = TimeSliderController->GetTickResolution();
-
-	double CurrentTime = TickResolution.AsSeconds(TimeSliderController->GetScrubPosition());
-
-	TOptional<double> NextTime;
-	TOptional<double> MinTime;
-
-	for (const TTuple<FCurveModelID, TUniquePtr<FCurveModel>>& Pair : CurveData)
-	{
-		FCurveModel* CurveModel = Pair.Value.Get();
-
-		if (CurveModel)
-		{
-			TArray<FKeyHandle> KeyHandles;
-			double MaxTime = NextTime.IsSet() ? NextTime.GetValue() : TNumericLimits<double>::Max();
-			CurveModel->GetKeys(CurrentTime, MaxTime, TNumericLimits<double>::Lowest(), TNumericLimits<double>::Max(), KeyHandles);
-
-			TArray<FKeyPosition> KeyPositions;
-			KeyPositions.SetNum(KeyHandles.Num());
-			CurveModel->GetKeyPositions(TArrayView<FKeyHandle>(KeyHandles), KeyPositions);
-
-			for (const FKeyPosition& KeyPosition : KeyPositions)
-			{
-				if (KeyPosition.InputValue > CurrentTime)
-				{
-					if (!NextTime.IsSet() || KeyPosition.InputValue < NextTime.GetValue())
-					{
-						NextTime = KeyPosition.InputValue;
-					}
-				}
-			}
-
-			double CurveMinTime, CurveMaxTime;
-			CurveModel->GetTimeRange(CurveMinTime, CurveMaxTime);
-			if (!MinTime.IsSet() || CurveMinTime < MinTime.GetValue())
-			{
-				MinTime = CurveMinTime;
-			}
-		}
-	}
-
-	if (NextTime.IsSet())
-	{
-		TimeSliderController->SetScrubPosition(NextTime.GetValue() * TickResolution,/*bEvaluate*/ true);
-	}
-	else if (MinTime.IsSet())
-	{
-		TimeSliderController->SetScrubPosition(MinTime.GetValue() * TickResolution, /*bEvaluate*/ true);
-	}
-}
-
-void FCurveEditor::StepToPreviousKey()
-{
-	TSharedPtr<ITimeSliderController> TimeSliderController = WeakTimeSliderController.Pin();
-	if (!TimeSliderController.IsValid())
-	{
-		return;
-	}
-
-	FFrameRate TickResolution = TimeSliderController->GetTickResolution();
-
-	double CurrentTime = TickResolution.AsSeconds(TimeSliderController->GetScrubPosition());
-
-	TOptional<double> PreviousTime;
-	TOptional<double> MaxTime;
-
-	for (const TTuple<FCurveModelID, TUniquePtr<FCurveModel>>& Pair : CurveData)
-	{
-		FCurveModel* CurveModel = Pair.Value.Get();
-
-		if (CurveModel)
-		{
-			TArray<FKeyHandle> KeyHandles;
-			double MinTime = PreviousTime.IsSet() ? PreviousTime.GetValue() : TNumericLimits<double>::Lowest();
-			CurveModel->GetKeys(MinTime, CurrentTime, TNumericLimits<double>::Lowest(), TNumericLimits<double>::Max(), KeyHandles);
-
-			TArray<FKeyPosition> KeyPositions;
-			KeyPositions.SetNum(KeyHandles.Num());
-			CurveModel->GetKeyPositions(TArrayView<FKeyHandle>(KeyHandles), KeyPositions);
-
-			for (const FKeyPosition& KeyPosition : KeyPositions)
-			{
-				if (KeyPosition.InputValue < CurrentTime)
-				{
-					if (!PreviousTime.IsSet() || KeyPosition.InputValue > PreviousTime.GetValue())
-					{
-						PreviousTime = KeyPosition.InputValue;
-					}
-				}
-			}
-
-			double CurveMinTime, CurveMaxTime;
-			CurveModel->GetTimeRange(CurveMinTime, CurveMaxTime);
-			if (!MaxTime.IsSet() || CurveMaxTime > MaxTime.GetValue())
-			{
-				MaxTime = CurveMaxTime;
-			}
-		}
-	}
-
-	if (PreviousTime.IsSet())
-	{
-		TimeSliderController->SetScrubPosition(PreviousTime.GetValue() * TickResolution,/*bEvaluate*/ true);
-	}
-	else if (MaxTime.IsSet())
-	{
-		TimeSliderController->SetScrubPosition(MaxTime.GetValue() * TickResolution, /*bEvaluate*/ true);
-	}
-}
-
-
-void FCurveEditor::StepForward()
-{
-	TSharedPtr<ITimeSliderController> TimeSliderController = WeakTimeSliderController.Pin();
-	if (!TimeSliderController.IsValid())
-	{
-		return;
-	}
-
-	FFrameRate TickResolution = TimeSliderController->GetTickResolution();
-	FFrameRate DisplayRate = TimeSliderController->GetDisplayRate();
-
-	FFrameTime OneFrame = FFrameRate::TransformTime(FFrameTime(1), DisplayRate, TickResolution);
-
-	TimeSliderController->SetScrubPosition(TimeSliderController->GetScrubPosition() + OneFrame, /*bEvaluate*/ true);
-}
-
-void FCurveEditor::StepBackward()
-{
-	TSharedPtr<ITimeSliderController> TimeSliderController = WeakTimeSliderController.Pin();
-	if (!TimeSliderController.IsValid())
-	{
-		return;
-	}
-
-	FFrameRate TickResolution = TimeSliderController->GetTickResolution();
-	FFrameRate DisplayRate = TimeSliderController->GetDisplayRate();
-
-	FFrameTime OneFrame = FFrameRate::TransformTime(FFrameTime(1), DisplayRate, TickResolution);
-
-	TimeSliderController->SetScrubPosition(TimeSliderController->GetScrubPosition() - OneFrame, /*bEvaluate*/ true);
-}
-
-void FCurveEditor::JumpToStart()
-{
-	TSharedPtr<ITimeSliderController> TimeSliderController = WeakTimeSliderController.Pin();
-	if (!TimeSliderController.IsValid())
-	{
-		return;
-	}
-
-	TimeSliderController->SetScrubPosition(TimeSliderController->GetTimeBounds().GetLowerBoundValue(), /*bEvaluate*/ true);
-}
-
-void FCurveEditor::JumpToEnd()
-{
-	TSharedPtr<ITimeSliderController> TimeSliderController = WeakTimeSliderController.Pin();
-	if (!TimeSliderController.IsValid())
-	{
-		return;
-	}
-
-	const bool bInsetDisplayFrame = IsInputSnappingEnabled();
-
-	FFrameRate TickResolution = TimeSliderController->GetTickResolution();
-	FFrameRate DisplayRate = TimeSliderController->GetDisplayRate();
-
-	// Calculate an offset from the end to go to. If they have snapping on (and the scrub style is a block) the last valid frame is represented as one
-	// whole display rate frame before the end, otherwise we just subtract a single frame which matches the behavior of hitting play and letting it run to the end.
-	FFrameTime OneFrame = bInsetDisplayFrame ? FFrameRate::TransformTime(FFrameTime(1), DisplayRate, TickResolution) : FFrameTime(1);
-	FFrameTime NewTime = TimeSliderController->GetTimeBounds().GetUpperBoundValue() - OneFrame;
-
-	TimeSliderController->SetScrubPosition(NewTime, /*bEvaluate*/ true);
-}
-
 void FCurveEditor::SetSelectionRangeStart()
 {
 	TSharedPtr<ITimeSliderController> TimeSliderController = WeakTimeSliderController.Pin();
diff --git a/Engine/Source/Editor/CurveEditor/Private/CurveEditorCommands.cpp b/Engine/Source/Editor/CurveEditor/Private/CurveEditorCommands.cpp
--- a/Engine/Source/Editor/CurveEditor/Private/CurveEditorCommands.cpp
+++ b/Engine/Source/Editor/CurveEditor/Private/CurveEditorCommands.cpp
@@ -118,12 +118,6 @@ void FCurveEditorCommands::RegisterCommands()
 	UI_COMMAND(TranslateSelectedKeysRight, "Translate Selected Keys Right", "Translate selected keys one frame to the right", EUserInterfaceActionType::Button, FInputChord(EModifierKey::Control, EKeys::Right));
 
 	// Time Management
-	UI_COMMAND(StepToNextKey, "Step to Next Key", "Step to the next key", EUserInterfaceActionType::Button, FInputChord(EKeys::Period));
-	UI_COMMAND(StepToPreviousKey, "Step to Previous Key", "Step to the previous key", EUserInterfaceActionType::Button, FInputChord(EKeys::Comma));
-	UI_COMMAND(StepForward, "Step Forward", "Step the timeline forward", EUserInterfaceActionType::Button, FInputChord(EKeys::Right));
-	UI_COMMAND(StepBackward, "Step Backward", "Step the timeline backward", EUserInterfaceActionType::Button, FInputChord(EKeys::Left));
-	UI_COMMAND(JumpToStart, "Jump to Start", "Jump to the start of the playback range", EUserInterfaceActionType::Button, FInputChord(EKeys::Up));
-	UI_COMMAND(JumpToEnd, "Jump to End", "Jump to the end of the playback range", EUserInterfaceActionType::Button, FInputChord(EModifierKey::Control, EKeys::Up));
 	UI_COMMAND(ScrubTime, "Scrub Time", "Scrub mouse left and right to change time", EUserInterfaceActionType::Button, FInputChord(EKeys::B));
 
 	// Selection Range
diff --git a/Engine/Source/Editor/CurveEditor/Public/CurveEditor.h b/Engine/Source/Editor/CurveEditor/Public/CurveEditor.h
--- a/Engine/Source/Editor/CurveEditor/Public/CurveEditor.h
+++ b/Engine/Source/Editor/CurveEditor/Public/CurveEditor.h
@@ -257,16 +257,6 @@ public:
 	/** Snap time to the first selected key */
 	void SnapToSelectedKey();
 
-	/** Step to next or previous key from the current time */
-	void StepToNextKey();
-	void StepToPreviousKey();
-
-	/** Step forward/backward, jump to start/end */
-	void StepForward();
-	void StepBackward();
-	void JumpToStart();
-	void JumpToEnd();
-
 	/** Selection range for ie. looping playback */
 	void SetSelectionRangeStart();
 	void SetSelectionRangeEnd();
diff --git a/Engine/Source/Editor/CurveEditor/Public/CurveEditorCommands.h b/Engine/Source/Editor/CurveEditor/Public/CurveEditorCommands.h
--- a/Engine/Source/Editor/CurveEditor/Public/CurveEditorCommands.h
+++ b/Engine/Source/Editor/CurveEditor/Public/CurveEditorCommands.h
@@ -126,12 +126,6 @@ public:
 	TSharedPtr<FUICommandInfo> TranslateSelectedKeysRight;
 
 	// Time Management
-	TSharedPtr<FUICommandInfo> StepToNextKey;
-	TSharedPtr<FUICommandInfo> StepToPreviousKey;
-	TSharedPtr<FUICommandInfo> StepForward;
-	TSharedPtr<FUICommandInfo> StepBackward;
-	TSharedPtr<FUICommandInfo> JumpToStart;
-	TSharedPtr<FUICommandInfo> JumpToEnd;
 	TSharedPtr<FUICommandInfo> ScrubTime;
 
 	// Selection Range
