Curve Editor: Fix failing ensure when FScopedSelectionTransaction is active and some system calls GUndo->Cancel.

The only known repro for this is the attached jira ticket, which was accidentally fixed by the previous CL 41600456 by changing logic FCurveEditorExtension::RequestSyncSelection.
FCurveEditorExtension::RequestSyncSelection was the only call site of this causing the failing ensure. This additional change does not directly contribute to fixing the issue in the reference jira ticket but it does future proof any future code that may happen to call GUndo->Cancel while FScopedSelectionTransaction is active.

diff --git a/Engine/Source/Editor/CurveEditor/Private/Modification/Utils/ScopedSelectionTransaction.cpp b/Engine/Source/Editor/CurveEditor/Private/Modification/Utils/ScopedSelectionTransaction.cpp
--- a/Engine/Source/Editor/CurveEditor/Private/Modification/Utils/ScopedSelectionTransaction.cpp
+++ b/Engine/Source/Editor/CurveEditor/Private/Modification/Utils/ScopedSelectionTransaction.cpp
@@ -27,7 +27,10 @@ FScopedSelectionTransaction::FScopedSelectionTransaction(TWeakPtr<FCurveEditor>
 
 FScopedSelectionTransaction::~FScopedSelectionTransaction()
 {
-	if (!IsOutstanding())
+	if (!IsOutstanding()
+		// This happens when this is a sub-transaction (Index > 1) and somebody cancels the root transaction (GUndo->Cancel()).
+		// That's technically a valid use of API. Now, calling AppendChange no longer makes sense as the owning root transaction was cancelled. 
+		|| !GUndo)
 	{
 		return;
 	}
diff --git a/Engine/Source/Editor/Sequencer/Private/MVVM/CurveEditorExtension.cpp b/Engine/Source/Editor/Sequencer/Private/MVVM/CurveEditorExtension.cpp
--- a/Engine/Source/Editor/Sequencer/Private/MVVM/CurveEditorExtension.cpp
+++ b/Engine/Source/Editor/Sequencer/Private/MVVM/CurveEditorExtension.cpp
@@ -650,12 +650,6 @@ void FCurveEditorExtension::RequestSyncSelection()
 			This->SyncSelection();
 		}
 
-		// Some external tools may reset GUndo in the meantime, e.g. USD Stage Editor during Actions->Import
-		if (!GUndo)
-		{
-			SelectionChange->Cancel();
-		}
-
 		SelectionChange.Reset();
 	});
 }
