Curve Editor: Fix filter label and description generated using UMETA entries not being localized.

diff --git a/Engine/Source/Editor/CurveEditor/Private/Filters/CurveEditorFilterBase.cpp b/Engine/Source/Editor/CurveEditor/Private/Filters/CurveEditorFilterBase.cpp
--- a/Engine/Source/Editor/CurveEditor/Private/Filters/CurveEditorFilterBase.cpp
+++ b/Engine/Source/Editor/CurveEditor/Private/Filters/CurveEditorFilterBase.cpp
@@ -12,13 +12,31 @@
 FText UCurveEditorFilterBase::GetLabel(const TSubclassOf<UCurveEditorFilterBase>& InClass)
 {
 	const FString* Label = InClass ? InClass->FindMetaData(TEXT("CurveEditorLabel")) : nullptr;
-	return Label ? FText::FromString(*Label) : LOCTEXT("Filter.LabelEmpty", "Filter");
+	if (!Label || !ensure(!InClass->GetName().IsEmpty()))
+	{
+		return LOCTEXT("Filter.LabelEmpty", "Filter");
+	}
+	
+	return FText::AsLocalizable_Advanced(
+		TEXT("CurveEditorFilterBase"),
+		*FString::Printf(TEXT("Filter.%s.Label"), *InClass->GetName()),
+		*Label
+		);
 }
 
 FText UCurveEditorFilterBase::GetDescription(const TSubclassOf<UCurveEditorFilterBase>& InClass)
 {
 	const FString* Description = InClass ? InClass->FindMetaData(TEXT("CurveEditorDescription")) : nullptr;
-	return Description ? FText::FromString(*Description) : FText::GetEmpty();
+	if (!Description || !ensure(!InClass->GetName().IsEmpty()))
+	{
+		return FText::GetEmpty();
+	}
+	
+	return FText::AsLocalizable_Advanced(
+		TEXT("CurveEditorFilterBase"),
+		*FString::Printf(TEXT("Filter.%s.Description"), *InClass->GetName()),
+		*Description
+		);
 }
 
 FSlateIcon UCurveEditorFilterBase::GetIcon(const TSubclassOf<UCurveEditorFilterBase>& InClass)
