Curve Editor: Fix context menu showing "Swap 0 Buffered Curves onto..." and "Apply 0 Buffered Curves..." when there are in fact more than 0 buffered curves.

diff --git a/Engine/Source/Editor/CurveEditor/Private/CurveEditorContextMenu.cpp b/Engine/Source/Editor/CurveEditor/Private/CurveEditorContextMenu.cpp
--- a/Engine/Source/Editor/CurveEditor/Private/CurveEditorContextMenu.cpp
+++ b/Engine/Source/Editor/CurveEditor/Private/CurveEditorContextMenu.cpp
@@ -36,17 +36,13 @@ public:
 	const TOptional<FCurveModelID> HoveredCurveID;
 	
 	const int32 NumSelectedKeys = CurveEditor->GetSelection().Count();
-	const TSharedPtr<FCurveEditor> LocalCurveEditor = CurveEditor->AsShared();
-	const int32 NumActiveCurves = LocalCurveEditor->GetSelectionFromTreeAndKeys().Num();
-	int32 NumBufferedCurves = 0;
+	const int32 NumActiveCurves = CurveEditor->GetSelectionFromTreeAndKeys().Num();
+	const int32 NumBufferedCurves;
 	
 	// We change the name to reflect the current number of curves selected.
-	const TAttribute<FText> BufferedCurvesText
-		= FText::Format(LOCTEXT("BufferCurvesContextMenu", "Buffer {0} Curves"), NumActiveCurves);
-	const TAttribute<FText> SwapBufferedCurvesText
-		= FText::Format(LOCTEXT("SwapBufferedCurvesContextMenu", "Swap {0} Buffered Curves onto {1} Selected Curves"), NumBufferedCurves, NumActiveCurves);
-	const TAttribute<FText> ApplyBufferedCurvesText
-		= FText::Format(LOCTEXT("ApplyBufferedCurvesContextMenu", "Apply {0} Buffered Curves onto {1} Selected Curves"), NumBufferedCurves, NumActiveCurves);
+	const TAttribute<FText> BufferedCurvesText;
+	const TAttribute<FText> SwapBufferedCurvesText;
+	const TAttribute<FText> ApplyBufferedCurvesText;
 
 	const FCurveModel* HoveredCurve = HoveredCurveID.IsSet() ? CurveEditor->FindCurve(HoveredCurveID.GetValue()) : nullptr;
 	
@@ -54,15 +50,22 @@ public:
 		: CurveEditor(CurveEditor)
 		, ClickedPoint(ClickedPoint)
 		, HoveredCurveID(HoveredCurveID)
-	{
-		for (const TUniquePtr<IBufferedCurveModel>& BufferedCurve : LocalCurveEditor->GetBufferedCurves())
+		, NumBufferedCurves([&CurveEditor]
 		{
-			if (LocalCurveEditor->IsActiveBufferedCurve(BufferedCurve))
+			int32 NumBufferedCurves = 0;
+			for (const TUniquePtr<IBufferedCurveModel>& BufferedCurve : CurveEditor->GetBufferedCurves())
 			{
-				++NumBufferedCurves;
+				if (CurveEditor->IsActiveBufferedCurve(BufferedCurve))
+				{
+					++NumBufferedCurves;
+				}
 			}
-		}
-	}
+			return NumBufferedCurves;
+		}())
+		, BufferedCurvesText(FText::Format(LOCTEXT("BufferCurvesContextMenu", "Buffer {0} Curves"), NumActiveCurves))
+		, SwapBufferedCurvesText(FText::Format(LOCTEXT("SwapBufferedCurvesContextMenu", "Swap {0} Buffered Curves onto {1} Selected Curves"), NumBufferedCurves, NumActiveCurves))
+		, ApplyBufferedCurvesText(FText::Format(LOCTEXT("ApplyBufferedCurvesContextMenu", "Apply {0} Buffered Curves onto {1} Selected Curves"), NumBufferedCurves, NumActiveCurves))
+	{}
 
 	void Build(FMenuBuilder& MenuBuilder)
 	{
@@ -302,4 +305,4 @@ void FCurveEditorContextMenu::BuildMenu(FMenuBuilder& MenuBuilder, TSharedRef<FC
 	Builder.Build(MenuBuilder);
 }
 
-#undef LOCTEXT_NAMESPACE
+#undef LOCTEXT_NAMESPACE
\ No newline at end of file
