Curve Editor: Fix static analysis.

diff --git a/Engine/Source/Editor/CurveEditor/Private/CurveEditorContextMenu.cpp b/Engine/Source/Editor/CurveEditor/Private/CurveEditorContextMenu.cpp
--- a/Engine/Source/Editor/CurveEditor/Private/CurveEditorContextMenu.cpp
+++ b/Engine/Source/Editor/CurveEditor/Private/CurveEditorContextMenu.cpp
@@ -37,12 +37,12 @@ public:
 	
 	const int32 NumSelectedKeys = CurveEditor->GetSelection().Count();
 	const int32 NumActiveCurves = CurveEditor->GetSelectionFromTreeAndKeys().Num();
-	const int32 NumBufferedCurves;
+	int32 NumBufferedCurves = 0;
 	
 	// We change the name to reflect the current number of curves selected.
-	const TAttribute<FText> BufferedCurvesText;
-	const TAttribute<FText> SwapBufferedCurvesText;
-	const TAttribute<FText> ApplyBufferedCurvesText;
+	TAttribute<FText> BufferedCurvesText;
+	TAttribute<FText> SwapBufferedCurvesText;
+	TAttribute<FText> ApplyBufferedCurvesText;
 
 	const FCurveModel* HoveredCurve = HoveredCurveID.IsSet() ? CurveEditor->FindCurve(HoveredCurveID.GetValue()) : nullptr;
 	
@@ -50,22 +50,23 @@ public:
 		: CurveEditor(CurveEditor)
 		, ClickedPoint(ClickedPoint)
 		, HoveredCurveID(HoveredCurveID)
-		, NumBufferedCurves([&CurveEditor]
+	{
+		for (const TUniquePtr<IBufferedCurveModel>& BufferedCurve : CurveEditor->GetBufferedCurves())
 		{
-			int32 NumBufferedCurves = 0;
-			for (const TUniquePtr<IBufferedCurveModel>& BufferedCurve : CurveEditor->GetBufferedCurves())
+			if (CurveEditor->IsActiveBufferedCurve(BufferedCurve))
 			{
-				if (CurveEditor->IsActiveBufferedCurve(BufferedCurve))
-				{
-					++NumBufferedCurves;
-				}
+				++NumBufferedCurves;
 			}
-			return NumBufferedCurves;
-		}())
-		, BufferedCurvesText(FText::Format(LOCTEXT("BufferCurvesContextMenu", "Buffer {0} Curves"), NumActiveCurves))
-		, SwapBufferedCurvesText(FText::Format(LOCTEXT("SwapBufferedCurvesContextMenu", "Swap {0} Buffered Curves onto {1} Selected Curves"), NumBufferedCurves, NumActiveCurves))
-		, ApplyBufferedCurvesText(FText::Format(LOCTEXT("ApplyBufferedCurvesContextMenu", "Apply {0} Buffered Curves onto {1} Selected Curves"), NumBufferedCurves, NumActiveCurves))
-	{}
+		}
+		
+		BufferedCurvesText = FText::Format(LOCTEXT("BufferCurvesContextMenu", "Buffer {0} Curves"), NumActiveCurves);
+		SwapBufferedCurvesText = FText::Format(
+			LOCTEXT("SwapBufferedCurvesContextMenu", "Swap {0} Buffered Curves onto {1} Selected Curves"), NumBufferedCurves, NumActiveCurves
+			);
+		ApplyBufferedCurvesText = FText::Format(
+			LOCTEXT("ApplyBufferedCurvesContextMenu", "Apply {0} Buffered Curves onto {1} Selected Curves"), NumBufferedCurves, NumActiveCurves
+			);
+	}
 
 	void Build(FMenuBuilder& MenuBuilder)
 	{
