Curve Editor: Refresh view after buffering curves.

Suppose you buffer a curve, move the buffered curve, and buffer the same curve again. Without this fix, the buffered curve continues to be drawn where it was before you do something that causes the viewport to be redrawn.

diff --git a/Engine/Source/Editor/CurveEditor/Private/Views/SInteractiveCurveEditorView.cpp b/Engine/Source/Editor/CurveEditor/Private/Views/SInteractiveCurveEditorView.cpp
--- a/Engine/Source/Editor/CurveEditor/Private/Views/SInteractiveCurveEditorView.cpp
+++ b/Engine/Source/Editor/CurveEditor/Private/Views/SInteractiveCurveEditorView.cpp
@@ -2003,6 +2003,12 @@ void SInteractiveCurveEditorView::BufferCurves()
 	if (CurveEditor.IsValid())
 	{
 		CurveEditor->AddBufferedCurves(CurveEditor->GetSelectionFromTreeAndKeys());
+
+		// If the user had previously buffered the curve, moved it, and now re-buffers it again, we need to trigger a re-draw.
+		if (CurveEditor->GetSettings()->GetShowBufferedCurves())
+		{
+			RefreshRetainer(); 
+		}
 	}
 }
 
