Curve Editor: Add tooltip to "select all" tree selection button. Fyi, this is the button shown on an item when you click a key of a curve.

diff --git a/Engine/Source/Editor/CurveEditor/Private/Tree/SCurveEditorTreeSelect.cpp b/Engine/Source/Editor/CurveEditor/Private/Tree/SCurveEditorTreeSelect.cpp
--- a/Engine/Source/Editor/CurveEditor/Private/Tree/SCurveEditorTreeSelect.cpp
+++ b/Engine/Source/Editor/CurveEditor/Private/Tree/SCurveEditorTreeSelect.cpp
@@ -34,6 +34,7 @@ void SCurveEditorTreeSelect::Construct(const FArguments& InArgs, TWeakPtr<FCurve
 	[
 		SNew(SButton)
 		.ButtonStyle(FAppStyle::Get(), "HoverHintOnly")
+		.ToolTipText(LOCTEXT("SelectAll", "Select all keys"))
 		.Visibility(this, &SCurveEditorTreeSelect::GetSelectVisibility)
 		.OnClicked(this, &SCurveEditorTreeSelect::SelectAll)
 		[
