Curve Editor: Fix that Ctrl + Y deselects current control rig control when you're on the top of the undo stack.

I didn't further debug it but apparently there is a subtle difference between calling GEditor->RedoTransaction, which FCurveEditor was doing, and the implementation of the main frame, which calls "FMainFrameActionCallbacks::ExecuteExecCommand(TEXT("TRANSACTION REDO"))".
Anyhow, this change should be safe. When you press Ctrl+Y while Curve Editor tab has focus, the key binding will just bubble up to the main frame, which then handles the redo command.

diff --git a/Engine/Source/Editor/CurveEditor/Private/CurveEditor.cpp b/Engine/Source/Editor/CurveEditor/Private/CurveEditor.cpp
--- a/Engine/Source/Editor/CurveEditor/Private/CurveEditor.cpp
+++ b/Engine/Source/Editor/CurveEditor/Private/CurveEditor.cpp
@@ -449,8 +449,6 @@ void FCurveEditor::BindCommands()
 	using namespace UE::CurveEditor;
 	UCurveEditorSettings* CurveSettings = Settings;
 		
-	CommandList->MapAction(FGenericCommands::Get().Undo,   FExecuteAction::CreateLambda([]{ GEditor->UndoTransaction(); }));
-	CommandList->MapAction(FGenericCommands::Get().Redo,   FExecuteAction::CreateLambda([]{ GEditor->RedoTransaction(); }));
 	CommandList->MapAction(FGenericCommands::Get().Delete, FExecuteAction::CreateSP(this, &FCurveEditor::DeleteSelection));
 
 	CommandList->MapAction(FGenericCommands::Get().Cut, FExecuteAction::CreateSP(this, &FCurveEditor::CutSelection));
