Curve Editor: Indiactor lines are now shown in normalized and stacked views, as well.

diff --git a/Engine/Source/Editor/CurveEditor/Private/Views/SCurveEditorViewNormalized.cpp b/Engine/Source/Editor/CurveEditor/Private/Views/SCurveEditorViewNormalized.cpp
--- a/Engine/Source/Editor/CurveEditor/Private/Views/SCurveEditorViewNormalized.cpp
+++ b/Engine/Source/Editor/CurveEditor/Private/Views/SCurveEditorViewNormalized.cpp
@@ -145,6 +145,7 @@ void SCurveEditorViewNormalized::PaintView(const FPaintArgs& Args, const FGeomet
 		DrawGridLines(CurveEditor.ToSharedRef(), AllottedGeometry, OutDrawElements, BaseLayerId, DrawEffects);
 		DrawBufferedCurves(AllottedGeometry, MyCullingRect, OutDrawElements, BaseLayerId, DrawEffects);
 		DrawCurves(CurveEditor.ToSharedRef(), AllottedGeometry, MyCullingRect, OutDrawElements, BaseLayerId, InWidgetStyle, DrawEffects);
+		DrawValueIndicatorLines(CurveEditor.ToSharedRef(), AllottedGeometry, OutDrawElements, BaseLayerId);
 	}
 }
 
diff --git a/Engine/Source/Editor/CurveEditor/Private/Views/SCurveEditorViewStacked.cpp b/Engine/Source/Editor/CurveEditor/Private/Views/SCurveEditorViewStacked.cpp
--- a/Engine/Source/Editor/CurveEditor/Private/Views/SCurveEditorViewStacked.cpp
+++ b/Engine/Source/Editor/CurveEditor/Private/Views/SCurveEditorViewStacked.cpp
@@ -116,6 +116,7 @@ void SCurveEditorViewStacked::PaintView(const FPaintArgs& Args, const FGeometry&
 		DrawLabels(AllottedGeometry, MyCullingRect, OutDrawElements, BaseLayerId, DrawEffects);
 		DrawBufferedCurves(AllottedGeometry, MyCullingRect, OutDrawElements, BaseLayerId, DrawEffects);
 		DrawCurves(CurveEditor.ToSharedRef(), AllottedGeometry, MyCullingRect, OutDrawElements, BaseLayerId, InWidgetStyle, DrawEffects);
+		DrawValueIndicatorLines(CurveEditor.ToSharedRef(), AllottedGeometry, OutDrawElements, BaseLayerId);
 	}
 }
 
diff --git a/Engine/Source/Editor/CurveEditor/Private/Views/SInteractiveCurveEditorView.cpp b/Engine/Source/Editor/CurveEditor/Private/Views/SInteractiveCurveEditorView.cpp
--- a/Engine/Source/Editor/CurveEditor/Private/Views/SInteractiveCurveEditorView.cpp
+++ b/Engine/Source/Editor/CurveEditor/Private/Views/SInteractiveCurveEditorView.cpp
@@ -960,13 +960,13 @@ void SInteractiveCurveEditorView::DrawValueIndicatorLines(
 {
 	UpdatedKeysWithValueIndicatorLines(*InCurveEditor);
 	
-	const FCurveEditorScreenSpace ViewSpace = GetViewSpace();
 	const FCurveModel* CurveModel = ValueIndicatorLineDrawData ? InCurveEditor->FindCurve(ValueIndicatorLineDrawData->HighlightedCurve) : nullptr;
 	if (!CurveModel)
 	{
 		return;
 	}
 
+	const FCurveEditorScreenSpace ViewSpace = GetCurveSpace(ValueIndicatorLineDrawData->HighlightedCurve);
 	const auto DrawLine = [&](const FKeyHandle& Handle)
 	{
 		FKeyPosition Position;
