Curve Editor: Disable CurveEditor.LogVerticalZoomMultipliers and CurveEditor.LogHorizontalZoomMultipliers, which were accidentally left to true as default value and cause log spam when a user zooms in curve eidtor using alt+move mouse.

diff --git a/Engine/Source/Editor/CurveEditor/Private/DragOperations/CurveEditorDragOperation_Zoom.cpp b/Engine/Source/Editor/CurveEditor/Private/DragOperations/CurveEditorDragOperation_Zoom.cpp
--- a/Engine/Source/Editor/CurveEditor/Private/DragOperations/CurveEditorDragOperation_Zoom.cpp
+++ b/Engine/Source/Editor/CurveEditor/Private/DragOperations/CurveEditorDragOperation_Zoom.cpp
@@ -13,12 +13,12 @@
 
 TAutoConsoleVariable<bool> CVarLogVerticalZoomMultipliers(
 	TEXT("CurveEditor.LogVerticalZoomMultipliers"),
-	true,
+	false,
 	TEXT("Logs the zoom multipliers to make it easier for you to tweak FCurveEditorZoomScaleConfig::VerticalZoomScale")
 	);
 TAutoConsoleVariable<bool> CVarLogHorizontalZoomMultipliers(
 	TEXT("CurveEditor.LogHorizontalZoomMultipliers"),
-	true,
+	false,
 	TEXT("Logs the zoom multipliers to make it easier for you to tweak FCurveEditorZoomScaleConfig::HorizontalZoomScale")
 	);
 
