Curve Editor: Remove label for View Options

#rnx

diff --git a/Engine/Source/Editor/CurveEditor/Private/SCurveEditorPanel.cpp b/Engine/Source/Editor/CurveEditor/Private/SCurveEditorPanel.cpp
--- a/Engine/Source/Editor/CurveEditor/Private/SCurveEditorPanel.cpp
+++ b/Engine/Source/Editor/CurveEditor/Private/SCurveEditorPanel.cpp
@@ -1399,7 +1399,6 @@ void SCurveEditorPanel::BuildToolbar(FToolBarBuilder& InToolBarBuilder, TSharedP
 	InToolBarBuilder.BeginSection("View");
 	InToolBarBuilder.BeginStyleOverride("CurveEditorToolbar");
 	{
-		InToolBarBuilder.SetLabelVisibility(EVisibility::Visible); // Show label..
 		InToolBarBuilder.AddComboButton(
 			FUIAction(),
 			FOnGetContent::CreateSP(this, &SCurveEditorPanel::MakeCurveEditorCurveViewOptionsMenu, InBaseExtender),
