Curve Editor: Introduce UserTangents visibility. It only shows tangents of keys the user has changed.

diff --git a/Engine/Source/Editor/CurveEditor/Private/CurveEditor.cpp b/Engine/Source/Editor/CurveEditor/Private/CurveEditor.cpp
--- a/Engine/Source/Editor/CurveEditor/Private/CurveEditor.cpp
+++ b/Engine/Source/Editor/CurveEditor/Private/CurveEditor.cpp
@@ -519,14 +519,17 @@ void FCurveEditor::BindCommands()
 	{
 		FExecuteAction SetAllTangents          = FExecuteAction::CreateUObject(Settings, &UCurveEditorSettings::SetTangentVisibility, ECurveEditorTangentVisibility::AllTangents);
 		FExecuteAction SetSelectedKeyTangents  = FExecuteAction::CreateUObject(Settings, &UCurveEditorSettings::SetTangentVisibility, ECurveEditorTangentVisibility::SelectedKeys);
+		FExecuteAction SetUserTangents			= FExecuteAction::CreateUObject(Settings, &UCurveEditorSettings::SetTangentVisibility, ECurveEditorTangentVisibility::UserTangents);
 		FExecuteAction SetNoTangents           = FExecuteAction::CreateUObject(Settings, &UCurveEditorSettings::SetTangentVisibility, ECurveEditorTangentVisibility::NoTangents);
 
 		FIsActionChecked IsAllTangents         = FIsActionChecked::CreateLambda( [CurveSettings]{ return CurveSettings->GetTangentVisibility() == ECurveEditorTangentVisibility::AllTangents; } );
 		FIsActionChecked IsSelectedKeyTangents = FIsActionChecked::CreateLambda( [CurveSettings]{ return CurveSettings->GetTangentVisibility() == ECurveEditorTangentVisibility::SelectedKeys; } );
+		FIsActionChecked IsUserTangents			= FIsActionChecked::CreateLambda( [CurveSettings]{ return CurveSettings->GetTangentVisibility() == ECurveEditorTangentVisibility::UserTangents; } );
 		FIsActionChecked IsNoTangents          = FIsActionChecked::CreateLambda( [CurveSettings]{ return CurveSettings->GetTangentVisibility() == ECurveEditorTangentVisibility::NoTangents; } );
 
 		CommandList->MapAction(FCurveEditorCommands::Get().SetAllTangentsVisibility, SetAllTangents, FCanExecuteAction(), IsAllTangents);
 		CommandList->MapAction(FCurveEditorCommands::Get().SetSelectedKeysTangentVisibility, SetSelectedKeyTangents, FCanExecuteAction(), IsSelectedKeyTangents);
+		CommandList->MapAction(FCurveEditorCommands::Get().SetUserTangentsVisibility, SetUserTangents, FCanExecuteAction(), IsUserTangents);
 		CommandList->MapAction(FCurveEditorCommands::Get().SetNoTangentsVisibility, SetNoTangents, FCanExecuteAction(), IsNoTangents);
 	}
 
diff --git a/Engine/Source/Editor/CurveEditor/Private/CurveEditorCommands.cpp b/Engine/Source/Editor/CurveEditor/Private/CurveEditorCommands.cpp
--- a/Engine/Source/Editor/CurveEditor/Private/CurveEditorCommands.cpp
+++ b/Engine/Source/Editor/CurveEditor/Private/CurveEditorCommands.cpp
@@ -60,6 +60,7 @@ void FCurveEditorCommands::RegisterCommands()
 
 	// Tangent Visibility
 	UI_COMMAND(SetAllTangentsVisibility, "All Tangents", "Show all tangents in the curve editor.", EUserInterfaceActionType::RadioButton, FInputChord());
+	UI_COMMAND(SetUserTangentsVisibility, "User Tangents", "Show tangents that are user edited, i.e. set to user or break should be visible.", EUserInterfaceActionType::RadioButton, FInputChord());
 	UI_COMMAND(SetSelectedKeysTangentVisibility, "Selected Keys", "Show tangents for selected keys in the curve editor.", EUserInterfaceActionType::RadioButton, FInputChord());
 	UI_COMMAND(SetNoTangentsVisibility, "No Tangents", "Show no tangents in the curve editor.", EUserInterfaceActionType::RadioButton, FInputChord());
 
diff --git a/Engine/Source/Editor/CurveEditor/Private/CurveEditorCurveDrawParamsCache.cpp b/Engine/Source/Editor/CurveEditor/Private/CurveEditorCurveDrawParamsCache.cpp
--- a/Engine/Source/Editor/CurveEditor/Private/CurveEditorCurveDrawParamsCache.cpp
+++ b/Engine/Source/Editor/CurveEditor/Private/CurveEditorCurveDrawParamsCache.cpp
@@ -24,7 +24,7 @@ namespace UE::CurveEditor
 		TRACE_CPUPROFILER_EVENT_SCOPE(FCurveDrawParamsCache::Invalidate);
 
 		WeakCurveEditor = CurveEditorView->GetCurveEditor();
-
+		
 		UpdateCurveCacheFlags(CurveEditorView);
 		DrawCurves(CurveEditorView, ModelIDs);
 	}
@@ -106,7 +106,7 @@ namespace UE::CurveEditor
 
 		if (VisibleKeys.Num())
 		{
-			ECurveEditorTangentVisibility TangentVisibility = CurveEditor->GetSettings()->GetTangentVisibility();
+			const ECurveEditorTangentVisibility TangentVisibility = CurveEditor->GetSettings()->GetTangentVisibility();
 
 			TArray<FKeyPosition> AllKeyPositions;
 			TArray<FKeyAttributes> AllKeyAttributes;
@@ -122,12 +122,16 @@ namespace UE::CurveEditor
 				const FKeyPosition& KeyPosition = AllKeyPositions[Index];
 				const FKeyAttributes& Attributes = AllKeyAttributes[Index];
 
-				bool bShowTangents = TangentVisibility == ECurveEditorTangentVisibility::AllTangents ||
-					(TangentVisibility == ECurveEditorTangentVisibility::SelectedKeys && SelectedKeys &&
-						(SelectedKeys->Contains(VisibleKeys[Index], ECurvePointType::Any)));
+				static_assert(static_cast<int32>(ECurveEditorTangentVisibility::Num) == 4, "Update drawing algorithm");
+				const bool bShowTangents = TangentVisibility == ECurveEditorTangentVisibility::AllTangents
+					|| (TangentVisibility == ECurveEditorTangentVisibility::SelectedKeys && SelectedKeys &&
+						(SelectedKeys->Contains(VisibleKeys[Index], ECurvePointType::Any)))
+					|| (TangentVisibility == ECurveEditorTangentVisibility::UserTangents && Attributes.HasTangentMode()
+						&& (Attributes.GetTangentMode() == RCTM_User || Attributes.GetTangentMode() == RCTM_Break))
+				;
 
-				double TimeScreenPos = CurveSpace.SecondsToScreen(KeyPosition.InputValue + InputOffset);
-				double ValueScreenPos = CurveSpace.ValueToScreen(KeyPosition.OutputValue);
+				const double TimeScreenPos = CurveSpace.SecondsToScreen(KeyPosition.InputValue + InputOffset);
+				const double ValueScreenPos = CurveSpace.ValueToScreen(KeyPosition.OutputValue);
 
 				// Add this key
 				FCurvePointInfo Key(KeyHandle);
diff --git a/Engine/Source/Editor/CurveEditor/Private/SCurveEditorPanel.cpp b/Engine/Source/Editor/CurveEditor/Private/SCurveEditorPanel.cpp
--- a/Engine/Source/Editor/CurveEditor/Private/SCurveEditorPanel.cpp
+++ b/Engine/Source/Editor/CurveEditor/Private/SCurveEditorPanel.cpp
@@ -1190,6 +1190,7 @@ TSharedRef<SWidget> SCurveEditorPanel::MakeCurveEditorCurveViewOptionsMenu(TShar
 	MenuBuilder.BeginSection("TangentVisibility", LOCTEXT("CurveEditorMenuTangentVisibilityHeader", "Tangent Visibility"));
 	{
 		MenuBuilder.AddMenuEntry(FCurveEditorCommands::Get().SetAllTangentsVisibility);
+		MenuBuilder.AddMenuEntry(FCurveEditorCommands::Get().SetUserTangentsVisibility);
 		MenuBuilder.AddMenuEntry(FCurveEditorCommands::Get().SetSelectedKeysTangentVisibility);
 		MenuBuilder.AddMenuEntry(FCurveEditorCommands::Get().SetNoTangentsVisibility);
 	}
diff --git a/Engine/Source/Editor/CurveEditor/Public/CurveEditorCommands.h b/Engine/Source/Editor/CurveEditor/Public/CurveEditorCommands.h
--- a/Engine/Source/Editor/CurveEditor/Public/CurveEditorCommands.h
+++ b/Engine/Source/Editor/CurveEditor/Public/CurveEditorCommands.h
@@ -75,6 +75,7 @@ public:
 	TSharedPtr<FUICommandInfo> SetPostInfinityExtrapConstant;
 
 	TSharedPtr<FUICommandInfo> SetAllTangentsVisibility;
+	TSharedPtr<FUICommandInfo> SetUserTangentsVisibility;
 	TSharedPtr<FUICommandInfo> SetSelectedKeysTangentVisibility;
 	TSharedPtr<FUICommandInfo> SetNoTangentsVisibility;
 
diff --git a/Engine/Source/Editor/CurveEditor/Public/CurveEditorSettings.h b/Engine/Source/Editor/CurveEditor/Public/CurveEditorSettings.h
--- a/Engine/Source/Editor/CurveEditor/Public/CurveEditorSettings.h
+++ b/Engine/Source/Editor/CurveEditor/Public/CurveEditorSettings.h
@@ -26,7 +26,12 @@ enum class ECurveEditorTangentVisibility : uint8
 	/** Only tangents from selected keys should be visible. */
 	SelectedKeys,
 	/** Don't display tangents. */
-	NoTangents
+	NoTangents,
+	/** Only tangents that are user edited, i.e. set to user or break should be visible. */
+	UserTangents,
+
+	// Add new entries above.
+	Num UMETA(Hidden)
 };
 
 /** Defines the position to center the zoom about in the curve editor. */
diff --git a/Engine/Source/Editor/MovieSceneTools/Private/Cache/MovieSceneCachedCurve.h b/Engine/Source/Editor/MovieSceneTools/Private/Cache/MovieSceneCachedCurve.h
--- a/Engine/Source/Editor/MovieSceneTools/Private/Cache/MovieSceneCachedCurve.h
+++ b/Engine/Source/Editor/MovieSceneTools/Private/Cache/MovieSceneCachedCurve.h
@@ -489,14 +489,15 @@ namespace UE::MovieSceneTools
 				const FKeyAttributes* KeyAttributesPtr = nullptr;
 				if (KeyAttributes.IsValidIndex(DataIndex))
 				{
-					if (TangentVisibility == ECurveEditorTangentVisibility::SelectedKeys)
-					{
-						if (Selection.Contains(KeyHandle, ECurvePointType::Any))
-						{
-							KeyAttributesPtr = &KeyAttributes[DataIndex];
-						}
-					}
-					else if (TangentVisibility == ECurveEditorTangentVisibility::AllTangents)
+					static_assert(static_cast<int32>(ECurveEditorTangentVisibility::Num) == 4, "Update drawing algorithm");
+					const bool bSelectedKeys = TangentVisibility == ECurveEditorTangentVisibility::SelectedKeys
+						&& Selection.Contains(KeyHandle, ECurvePointType::Any);
+					const bool bUserTangents = TangentVisibility == ECurveEditorTangentVisibility::UserTangents
+						&& KeyAttributes[DataIndex].HasTangentMode()
+						&& (KeyAttributes[DataIndex].GetTangentMode() == RCTM_User || KeyAttributes[DataIndex].GetTangentMode() == RCTM_Break);
+					const bool bAllTangents = TangentVisibility == ECurveEditorTangentVisibility::AllTangents;
+					
+					if (bSelectedKeys || bUserTangents || bAllTangents)
 					{
 						KeyAttributesPtr = &KeyAttributes[DataIndex];
 					}
diff --git a/Engine/Source/Editor/UnrealEd/Private/SCurveEditor.cpp b/Engine/Source/Editor/UnrealEd/Private/SCurveEditor.cpp
--- a/Engine/Source/Editor/UnrealEd/Private/SCurveEditor.cpp
+++ b/Engine/Source/Editor/UnrealEd/Private/SCurveEditor.cpp
@@ -257,6 +257,11 @@ void SCurveEditor::Construct(const FArguments& InArgs)
 		FCanExecuteAction::CreateLambda( []{ return true; } ),
 		FIsActionChecked::CreateLambda( [this]{ return Settings->GetTangentVisibility() == ECurveEditorTangentVisibility::AllTangents; } ) );
 
+	Commands->MapAction(FCurveEditorCommands::Get().SetUserTangentsVisibility,
+		FExecuteAction::CreateLambda( [this]{ Settings->SetTangentVisibility( ECurveEditorTangentVisibility::UserTangents ); } ),
+		FCanExecuteAction::CreateLambda( []{ return true; } ),
+		FIsActionChecked::CreateLambda( [this]{ return Settings->GetTangentVisibility() == ECurveEditorTangentVisibility::UserTangents; } ) );
+	
 	Commands->MapAction(FCurveEditorCommands::Get().SetSelectedKeysTangentVisibility,
 		FExecuteAction::CreateLambda( [this]{ Settings->SetTangentVisibility( ECurveEditorTangentVisibility::SelectedKeys ); } ),
 		FCanExecuteAction::CreateLambda( []{ return true; } ),
