Curve Editor: Normalized view should be relative to time range

diff --git a/Engine/Source/Editor/CurveEditor/Private/Views/SCurveEditorViewNormalized.cpp b/Engine/Source/Editor/CurveEditor/Private/Views/SCurveEditorViewNormalized.cpp
--- a/Engine/Source/Editor/CurveEditor/Private/Views/SCurveEditorViewNormalized.cpp
+++ b/Engine/Source/Editor/CurveEditor/Private/Views/SCurveEditorViewNormalized.cpp
@@ -59,7 +59,7 @@ void SCurveEditorViewNormalized::Construct(const FArguments& InArgs, TWeakPtr<FC
 	];
 }
 
-FTransform2d CalculateViewToCurveTransform(const double OutputMin, const double OutputMax, const double InCurveOutputMin, const double InCurveOutputMax)
+FTransform2d CalculateViewToCurveTransform(const double InCurveOutputMin, const double InCurveOutputMax)
 {
 	const double Scale = (InCurveOutputMax - InCurveOutputMin);
 	if (InCurveOutputMax > InCurveOutputMin)
@@ -105,7 +105,7 @@ void SCurveEditorViewNormalized::DrawBufferedCurves(const FGeometry& AllottedGeo
 		FTransform2d ViewToBufferedCurveTransform;
 		double CurveOutputMin = BufferedCurve->GetValueMin(), CurveOutputMax = BufferedCurve->GetValueMax();
 
-		ViewToBufferedCurveTransform = CalculateViewToCurveTransform(OutputMin, OutputMax,CurveOutputMin, CurveOutputMax);
+		ViewToBufferedCurveTransform = CalculateViewToCurveTransform(CurveOutputMin, CurveOutputMax);
 
 		TArray<TTuple<double, double>> CurveSpaceInterpolatingPoints;
 		FCurveEditorScreenSpace CurveSpace = GetViewSpace().ToCurveSpace(ViewToBufferedCurveTransform);
@@ -183,15 +183,16 @@ void SCurveEditorViewNormalized::UpdateViewToTransformCurves(double InputMin, do
 			continue;
 		}
 		
-		// Consider the whole time range when getting the value range
-		double MinTime;
-		double MaxTime;
+		// Consider the whole time range when getting the value range...
+		double MinTime, MaxTime;
 		Curve->GetTimeRange(MinTime, MaxTime);
+		// ... but clamp it to the time range the user has set with the time controller
+		ClampToTimeController(MinTime, MaxTime);
 
 		double CurveOutputMin = 0, CurveOutputMax = 1;
 		Curve->GetValueRange(MinTime, MaxTime, CurveOutputMin, CurveOutputMax);
 
-		It->Value.ViewToCurveTransform = CalculateViewToCurveTransform(OutputMin, OutputMax, CurveOutputMin, CurveOutputMax);
+		It->Value.ViewToCurveTransform = CalculateViewToCurveTransform(CurveOutputMin, CurveOutputMax);
 	}
 }
 
@@ -225,3 +226,25 @@ void SCurveEditorViewNormalized::InternalUpdateViewToTransformCurves()
 		UpdateViewToTransformCurves(InputMin, InputMax);
 	}
 }
+
+void SCurveEditorViewNormalized::ClampToTimeController(double& OutMin, double& OutMax) const
+{
+	const TSharedPtr<ITimeSliderController> TimeController = GetCurveEditor()->GetTimeSliderController();
+	if (!TimeController)
+    {
+		return;
+    }
+	
+	const FFrameRate TickResolution = TimeController->GetTickResolution();
+	const TRange<FFrameNumber> Range = TimeController->GetPlayRange();
+	if (Range.HasLowerBound())
+	{
+		const double Seconds = TickResolution.AsSeconds(Range.GetLowerBoundValue().Value);
+		OutMin = FMath::Max(OutMin, Seconds);
+	}
+	if (Range.HasUpperBound())
+	{
+		const double Seconds = TickResolution.AsSeconds(Range.GetUpperBoundValue().Value);
+		OutMax = FMath::Min(OutMax, Seconds);
+	}
+}
diff --git a/Engine/Source/Editor/CurveEditor/Public/Views/SCurveEditorViewNormalized.h b/Engine/Source/Editor/CurveEditor/Public/Views/SCurveEditorViewNormalized.h
--- a/Engine/Source/Editor/CurveEditor/Public/Views/SCurveEditorViewNormalized.h
+++ b/Engine/Source/Editor/CurveEditor/Public/Views/SCurveEditorViewNormalized.h
@@ -45,4 +45,10 @@ protected:
 
 	/** Helper to call UpdateViewToTransformCurves from internal methods */
 	void InternalUpdateViewToTransformCurves();
+
+	/**
+	 * Limits OutMin and OutMax to the time controller playback range, i.e. OutMin >= TimeController's min and OutMax <= TimeController's max.
+	 * This corresponds green, left and red, right horizontal lines on the frame ruler.
+	 */
+	void ClampToTimeController(double& OutMin, double& OutMax) const;
 };
