Curve Editor: Remove unused Flip Curve Menu

#rnx

diff --git a/Engine/Source/Editor/CurveEditor/Private/SCurveEditorPanel.cpp b/Engine/Source/Editor/CurveEditor/Private/SCurveEditorPanel.cpp
--- a/Engine/Source/Editor/CurveEditor/Private/SCurveEditorPanel.cpp
+++ b/Engine/Source/Editor/CurveEditor/Private/SCurveEditorPanel.cpp
@@ -1433,7 +1433,14 @@ void SCurveEditorPanel::BuildToolbar(FToolBarBuilder& InToolBarBuilder, TSharedP
 
 	InToolBarBuilder.BeginSection("Key Details");
 	{
-		InToolBarBuilder.AddWidget(KeyDetailsView.ToSharedRef(), {}, NAME_None, true, {}, {}, MakeResizeParams(TEXT("KeyDetails")));
+		InToolBarBuilder.AddWidget(
+			SNew(SBox)
+			.HAlign(HAlign_Center)
+			.VAlign(VAlign_Center)
+			[
+				KeyDetailsView.ToSharedRef()				
+			]
+			, {}, NAME_None, true, {}, {}, MakeResizeParams(TEXT("KeyDetails")));
 	}
 	InToolBarBuilder.EndSection();
 
@@ -1616,140 +1623,6 @@ TSharedRef<SWidget> SCurveEditorPanel::MakeGridSpacingMenu()
 	return OutputSnapWidget;
 }
 
-TSharedRef<SWidget> SCurveEditorPanel::MakeFlipCurveMenu(FCurveEditor::ECurveFlipDirection Direction)
-{
-	FMenuBuilder MenuBuilder(true, nullptr);
-	MenuBuilder.BeginSection("Header", LOCTEXT("CurveEditorMenuCurveFlipgHeader", "Curve Flip Settings"));
-
-	FCurveEditor::FCurveFlipRangeSettings* RangeSetting = (Direction == FCurveEditor::ECurveFlipDirection::Horizontal) ? &CurveEditor->HorizontalCurveFlipRangeSettings : &CurveEditor->VerticalCurveFlipRangeSettings;
-	
-	// For the x axis, the range of keys and the range of curve should be the same
-	if (Direction == FCurveEditor::ECurveFlipDirection::Horizontal)
-	{
-		// Add menu entry: "Use key/curve range"
-		MenuBuilder.AddMenuEntry(
-			LOCTEXT("UseKeyCurveRange", "Use Key/Curve Range"),
-			LOCTEXT("UseKeyCurveRangeTooltip", "Flip curve within the range of keys/curve"),
-			FSlateIcon(),
-			FUIAction(
-				FExecuteAction::CreateLambda([RangeSetting]()
-					{
-						RangeSetting->RangeType = FCurveEditor::ECurveFlipRangeType::CurveRange;
-					}),
-				FCanExecuteAction(),
-				FIsActionChecked::CreateLambda([RangeSetting]() { return RangeSetting->RangeType == FCurveEditor::ECurveFlipRangeType::CurveRange; })
-			),
-			NAME_None,
-			EUserInterfaceActionType::RadioButton
-		);
-	}
-
-	// For the y axis, the range of curve can be different than range of keys if tangent is set
-	if (Direction == FCurveEditor::ECurveFlipDirection::Vertical)
-	{
-		// Add menu entry: "Use key range"
-		MenuBuilder.AddMenuEntry(
-			LOCTEXT("UseKeyRange", "Use Key Range"),
-			LOCTEXT("UseKeyRangeTooltip", "Flip curve within the range of the keys"),
-			FSlateIcon(),
-			FUIAction(
-				FExecuteAction::CreateLambda([RangeSetting]()
-					{
-						RangeSetting->RangeType = FCurveEditor::ECurveFlipRangeType::KeyRange;
-					}),
-				FCanExecuteAction(),
-				FIsActionChecked::CreateLambda([RangeSetting]() { return RangeSetting->RangeType == FCurveEditor::ECurveFlipRangeType::KeyRange; })
-			),
-			NAME_None,
-			EUserInterfaceActionType::RadioButton
-		);
-
-		// Add menu entry: "Use curve range"
-		MenuBuilder.AddMenuEntry(
-			LOCTEXT("UseCurveRange", "Use Curve Range"),
-			LOCTEXT("UseCurveRangeTooltip", "Flip curve within the range of the curve"),
-			FSlateIcon(),
-			FUIAction(
-				FExecuteAction::CreateLambda([RangeSetting]()
-					{
-						RangeSetting->RangeType = FCurveEditor::ECurveFlipRangeType::CurveRange;
-					}),
-				FCanExecuteAction(),
-				FIsActionChecked::CreateLambda([RangeSetting]() { return RangeSetting->RangeType == FCurveEditor::ECurveFlipRangeType::CurveRange; })
-			),
-			NAME_None,
-			EUserInterfaceActionType::RadioButton
-		);
-	}
-
-	// Add "Custom Range" menu entry
-	MenuBuilder.AddMenuEntry(
-		LOCTEXT("CustomRange", "Use Custom Range"),
-		LOCTEXT("CustomRangeTooltip", "Flip curve within the range of float inputs (Default range is 0 to 1)"),
-		FSlateIcon(),
-		FUIAction(
-			FExecuteAction::CreateLambda([RangeSetting]()
-				{
-					RangeSetting->RangeType = FCurveEditor::ECurveFlipRangeType::CustomRange;
-				}),
-			FCanExecuteAction(),
-			FIsActionChecked::CreateLambda([RangeSetting]() { return RangeSetting->RangeType == FCurveEditor::ECurveFlipRangeType::CustomRange; })
-		),
-		NAME_None,
-		EUserInterfaceActionType::RadioButton
-	);
-
-	// Create a custom widget that includes two input boxes: one for "Min" and one for "Max"
-	TSharedRef<SWidget> CustomRangeWidget =
-		SNew(SVerticalBox)
-		+ SVerticalBox::Slot()
-		.AutoHeight()
-		.Padding(FMargin(0, 6, 0, 0))
-		[
-			SNew(STextBlock)
-				.Text(LOCTEXT("MinRangeLabel", "Min:"))
-		]
-		+ SVerticalBox::Slot()
-		.AutoHeight()
-		[
-			SNew(SEditableTextBox)
-				.MinDesiredWidth(60.0f)
-				.Text_Lambda([RangeSetting]() -> FText {
-				return FText::AsNumber(RangeSetting->MinRange);
-					})
-				.OnTextCommitted_Lambda([RangeSetting](const FText& InText, ETextCommit::Type CommitType)
-					{
-						RangeSetting->MinRange = FCString::Atof(*InText.ToString());
-					})
-		]
-		+ SVerticalBox::Slot()
-		.AutoHeight()
-		.Padding(FMargin(0, 8, 0, 0))
-		[
-			SNew(STextBlock)
-				.Text(LOCTEXT("MaxRangeLabel", "Max:"))
-		]
-		+ SVerticalBox::Slot()
-		.AutoHeight()
-		[
-			SNew(SEditableTextBox)
-				.MinDesiredWidth(60.0f)
-				.Text_Lambda([RangeSetting]() -> FText {
-				return FText::AsNumber(RangeSetting->MaxRange);
-					})
-				.OnTextCommitted_Lambda([RangeSetting](const FText& InText, ETextCommit::Type CommitType)
-					{
-						RangeSetting->MaxRange = FCString::Atof(*InText.ToString());
-					})
-		];
-
-	// Only enable the custom range input widget when "Custom Range" is selected.
-	CustomRangeWidget->SetEnabled(TAttribute<bool>::CreateLambda([RangeSetting]() { return RangeSetting->RangeType == FCurveEditor::ECurveFlipRangeType::CustomRange; }));
-
-	MenuBuilder.AddWidget(CustomRangeWidget, LOCTEXT("CustomRangeInputs", "Range"));
-	return MenuBuilder.MakeWidget();
-}
-
 TSharedRef<SWidget> SCurveEditorPanel::MakeAxisSnapMenu()
 {
 	FMenuBuilder MenuBuilder(true, GetCommands());
diff --git a/Engine/Source/Editor/CurveEditor/Public/SCurveEditorPanel.h b/Engine/Source/Editor/CurveEditor/Public/SCurveEditorPanel.h
--- a/Engine/Source/Editor/CurveEditor/Public/SCurveEditorPanel.h
+++ b/Engine/Source/Editor/CurveEditor/Public/SCurveEditorPanel.h
@@ -208,9 +208,6 @@ private:
 	TSharedRef<SWidget> MakeGridSpacingMenu();
 	TSharedRef<SWidget> MakeAxisSnapMenu();
 
-	/* Curve Flip make menu function and settings */
-	TSharedRef<SWidget> MakeFlipCurveMenu(FCurveEditor::ECurveFlipDirection Direction);
-
 	bool IsInlineEditPanelEditable() const;
 	EVisibility ShouldInstructionOverlayBeVisible() const;
 
