Curve Editor: Adjust FScopedSelectionTransaction to no longer inherit from FScopedTransaction, which is not designed for inheritance.
- Renamed FScopedSelectionTransaction to FScopedSelectionChange (avoid implying it's a transaction during the scope)
- Renamed FSelectionDeltaChange to FSelectionChangeData (better communicate it is data only without logic)
- CurveEditor::FTransactionManager::AppendChange now starts a FScopedTransaction internally. This is how it works with FEdModeToolsContextTransactionImpl::AppendChange, too.

diff --git a/Engine/Source/Editor/CurveEditor/Private/CurveEditor.cpp b/Engine/Source/Editor/CurveEditor/Private/CurveEditor.cpp
--- a/Engine/Source/Editor/CurveEditor/Private/CurveEditor.cpp
+++ b/Engine/Source/Editor/CurveEditor/Private/CurveEditor.cpp
@@ -60,7 +60,7 @@
 #include "UObject/UnrealNames.h"
 #include "UnrealExporter.h"
 #include "Misc/SmartSnap.h"
-#include "Modification/Utils/ScopedSelectionTransaction.h"
+#include "Modification/Utils/ScopedSelectionChange.h"
 #include "Widgets/Colors/SColorPicker.h"
 #include "Widgets/Notifications/SNotificationList.h"
 
@@ -1047,7 +1047,7 @@ void FCurveEditor::ClearSelectionRange()
 
 void FCurveEditor::SelectAllKeys()
 {
-	const UE::CurveEditor::FScopedSelectionTransaction Transaction(SharedThis(this), LOCTEXT("SelectAllKeys", "Select all keys"));
+	const UE::CurveEditor::FScopedSelectionChange Transaction(SharedThis(this), LOCTEXT("SelectAllKeys", "Select all keys"));
 	for (FCurveModelID ID : GetEditedCurves())
 	{
 		if (FCurveModel* Curve = FindCurve(ID))
@@ -1061,7 +1061,7 @@ void FCurveEditor::SelectAllKeys()
 
 void FCurveEditor::SelectForward()
 {
-	const UE::CurveEditor::FScopedSelectionTransaction Transaction(SharedThis(this), LOCTEXT("SelectForward", "Select forward"));
+	const UE::CurveEditor::FScopedSelectionChange Transaction(SharedThis(this), LOCTEXT("SelectForward", "Select forward"));
 	Selection.Clear();
 
 	TSharedPtr<ITimeSliderController> TimeSliderController = WeakTimeSliderController.Pin();
@@ -1087,7 +1087,7 @@ void FCurveEditor::SelectForward()
 
 void FCurveEditor::SelectBackward()
 {
-	const UE::CurveEditor::FScopedSelectionTransaction Transaction(SharedThis(this), LOCTEXT("SelectBackward", "Select backward"));
+	const UE::CurveEditor::FScopedSelectionChange Transaction(SharedThis(this), LOCTEXT("SelectBackward", "Select backward"));
 	Selection.Clear();
 
 	TSharedPtr<ITimeSliderController> TimeSliderController = WeakTimeSliderController.Pin();
@@ -1113,13 +1113,13 @@ void FCurveEditor::SelectBackward()
 
 void FCurveEditor::SelectNone()
 {
-	const UE::CurveEditor::FScopedSelectionTransaction Transaction(SharedThis(this));
+	const UE::CurveEditor::FScopedSelectionChange Transaction(SharedThis(this));
 	Selection.Clear();
 }
 
 void FCurveEditor::InvertSelection()
 {
-	const UE::CurveEditor::FScopedSelectionTransaction Transaction(SharedThis(this), LOCTEXT("InvertSelection", "Invert selection"));
+	const UE::CurveEditor::FScopedSelectionChange Transaction(SharedThis(this), LOCTEXT("InvertSelection", "Invert selection"));
 	for (const TTuple<FCurveModelID, FKeyHandleSet>& Pair : Selection.GetAll())
 	{
 		FCurveModelID CurveModelID = Pair.Key;
@@ -1846,7 +1846,7 @@ void FCurveEditor::PasteKeys(UE::CurveEditor::FKeyPasteArgs InArgs)
 		}
 	}
 
-	const UE::CurveEditor::FScopedSelectionTransaction KeyChange(SharedThis(this), LOCTEXT("PasteKeys", "Paste Keys"));
+	const UE::CurveEditor::FScopedSelectionChange KeyChange(SharedThis(this), LOCTEXT("PasteKeys", "Paste Keys"));
 	Selection.Clear();
 
 
@@ -1976,7 +1976,7 @@ void FCurveEditor::PasteKeys(UE::CurveEditor::FKeyPasteArgs InArgs)
 
 void FCurveEditor::DeleteSelection()
 {
-	const UE::CurveEditor::FScopedSelectionTransaction Transaction(SharedThis(this), LOCTEXT("DeleteKeys", "Delete Keys"));
+	const UE::CurveEditor::FScopedSelectionChange Transaction(SharedThis(this), LOCTEXT("DeleteKeys", "Delete Keys"));
 
 	double CurrentTime = 0.0;
 	if (TSharedPtr<ITimeSliderController> TimeSliderController = WeakTimeSliderController.Pin())
@@ -2143,7 +2143,7 @@ bool FCurveEditor::CanFlattenOrStraightenSelection() const
 
 void FCurveEditor::SmartSnapSelection()
 {
-	const UE::CurveEditor::FScopedSelectionTransaction Transaction(SharedThis(this), LOCTEXT("SmartSnapKeys", "Smart Snap"));
+	const UE::CurveEditor::FScopedSelectionChange Transaction(SharedThis(this), LOCTEXT("SmartSnapKeys", "Smart Snap"));
 	
 	TMap<FCurveModelID, FKeyHandleSet> OutKeysToSelect;
 	UE::CurveEditor::EnumerateSmartSnappableKeys(*this, Selection.GetAll(), OutKeysToSelect,
diff --git a/Engine/Source/Editor/CurveEditor/Private/CurveEditorSelectionPrivate.h b/Engine/Source/Editor/CurveEditor/Private/CurveEditorSelectionPrivate.h
--- a/Engine/Source/Editor/CurveEditor/Private/CurveEditorSelectionPrivate.h
+++ b/Engine/Source/Editor/CurveEditor/Private/CurveEditorSelectionPrivate.h
@@ -3,7 +3,7 @@
 #pragma once
 
 #include "CurveEditorSelection.h"
-#include "Modification/Utils/SelectionChangeUtils.h"
+#include "Modification/Selection/SelectionChangeUtils.h"
 
 namespace UE::CurveEditor::KeySelection 
 {
diff --git a/Engine/Source/Editor/CurveEditor/Private/DragOperations/CurveEditorDragOperation_Marquee.h b/Engine/Source/Editor/CurveEditor/Private/DragOperations/CurveEditorDragOperation_Marquee.h
--- a/Engine/Source/Editor/CurveEditor/Private/DragOperations/CurveEditorDragOperation_Marquee.h
+++ b/Engine/Source/Editor/CurveEditor/Private/DragOperations/CurveEditorDragOperation_Marquee.h
@@ -7,7 +7,7 @@
 #include "ICurveEditorDragOperation.h"
 #include "Layout/SlateRect.h"
 #include "Math/Vector2D.h"
-#include "Modification/Utils/ScopedSelectionTransaction.h"
+#include "Modification/Utils/ScopedSelectionChange.h"
 
 class FCurveEditor;
 class FSlateWindowElementList;
@@ -30,7 +30,7 @@ public:
 private:
 
 	/** Makes this change undo / redo able. */
-	const UE::CurveEditor::FScopedSelectionTransaction SelectionChange;
+	const UE::CurveEditor::FScopedSelectionChange SelectionChange;
 
 	/** The current marquee rectangle */
 	FSlateRect Marquee;
diff --git a/Engine/Source/Editor/CurveEditor/Private/Filters/FilterUtils.cpp b/Engine/Source/Editor/CurveEditor/Private/Filters/FilterUtils.cpp
--- a/Engine/Source/Editor/CurveEditor/Private/Filters/FilterUtils.cpp
+++ b/Engine/Source/Editor/CurveEditor/Private/Filters/FilterUtils.cpp
@@ -3,7 +3,7 @@
 #include "FilterUtils.h"
 
 #include "Filters/CurveEditorFilterBase.h"
-#include "Modification/Utils/ScopedSelectionTransaction.h"
+#include "Modification/Utils/ScopedSelectionChange.h"
 
 #include "Containers/Map.h"
 #include "Templates/SharedPointer.h"
@@ -24,7 +24,7 @@ namespace UE::CurveEditor::FilterUtils
 		TMap<FCurveModelID, FKeyHandleSet> OutKeysToSelect;
 		InFilter.ApplyFilter(InCurveEditor, InSelectedKeys, OutKeysToSelect);
 		
-		const FScopedSelectionTransaction Transaction(InCurveEditor);
+		const FScopedSelectionChange Transaction(InCurveEditor);
 		// Clear their selection and then set it to the keys the filter thinks you should have selected.
 		InCurveEditor->GetSelection().Clear();
 
diff --git a/Engine/Source/Editor/CurveEditor/Private/Modification/Changes/SelectionChangeCommand.cpp b/Engine/Source/Editor/CurveEditor/Private/Modification/Changes/SelectionChangeCommand.cpp
deleted file mode 100644
diff --git a/Engine/Source/Editor/CurveEditor/Private/Modification/Changes/SelectionChangeCommand.h b/Engine/Source/Editor/CurveEditor/Private/Modification/Changes/SelectionChangeCommand.h
deleted file mode 100644
diff --git a/Engine/Source/Editor/CurveEditor/Private/Modification/Changes/SelectionDeltaChange.h b/Engine/Source/Editor/CurveEditor/Private/Modification/Changes/SelectionDeltaChange.h
deleted file mode 100644
diff --git a/Engine/Source/Editor/CurveEditor/Private/Modification/Selection/SelectionChangeCommand.cpp b/Engine/Source/Editor/CurveEditor/Private/Modification/Selection/SelectionChangeCommand.cpp
new file mode 100644
--- /dev/null
+++ b/Engine/Source/Editor/CurveEditor/Private/Modification/Selection/SelectionChangeCommand.cpp
@@ -0,0 +1,25 @@
+﻿// Copyright Epic Games, Inc. All Rights Reserved.
+
+#include "SelectionChangeCommand.h"
+
+#include "CurveEditor.h"
+#include "Modification/Selection/SelectionChangeUtils.h"
+
+namespace UE::CurveEditor
+{
+void FSelectionChangeCommand::Apply(UObject* Object)
+{
+	if (const TSharedPtr<FCurveEditor> CurveEditor = GetCurveEditor(Object))
+	{
+		ApplySelectionChange(CurveEditor->Selection, DeltaChange);
+	}
+}
+
+void FSelectionChangeCommand::Revert(UObject* Object)
+{
+	if (const TSharedPtr<FCurveEditor> CurveEditor = GetCurveEditor(Object))
+	{
+		RevertSelectionChange(CurveEditor->Selection, DeltaChange);
+	}
+}
+}
diff --git a/Engine/Source/Editor/CurveEditor/Private/Modification/Selection/SelectionChangeCommand.h b/Engine/Source/Editor/CurveEditor/Private/Modification/Selection/SelectionChangeCommand.h
new file mode 100644
--- /dev/null
+++ b/Engine/Source/Editor/CurveEditor/Private/Modification/Selection/SelectionChangeCommand.h
@@ -0,0 +1,28 @@
+﻿// Copyright Epic Games, Inc. All Rights Reserved.
+
+#pragma once
+
+#include "SelectionDeltaChange.h"
+#include "Modification/CurveEditorCommandChange.h"
+
+namespace UE::CurveEditor
+{
+/** Changes that changes the curve editor's selection. */
+class FSelectionChangeCommand : public FCurveEditorCommandChange
+{
+public:
+
+	explicit FSelectionChangeCommand(FSelectionDeltaChange InDeltaChange) : DeltaChange(MoveTemp(InDeltaChange)) {}
+	
+	//~ Begin FChange Interface
+	virtual void Apply(UObject* Object) override;
+	virtual void Revert(UObject* Object) override;
+	//~ End FChange Interface
+
+private:
+
+	/** The change made to the selection. */
+	const FSelectionDeltaChange DeltaChange;
+};
+}
+
diff --git a/Engine/Source/Editor/CurveEditor/Private/Modification/Selection/SelectionChangeUtils.cpp b/Engine/Source/Editor/CurveEditor/Private/Modification/Selection/SelectionChangeUtils.cpp
new file mode 100644
--- /dev/null
+++ b/Engine/Source/Editor/CurveEditor/Private/Modification/Selection/SelectionChangeUtils.cpp
@@ -0,0 +1,145 @@
+﻿// Copyright Epic Games, Inc. All Rights Reserved.
+
+#include "SelectionChangeUtils.h"
+
+#include "CurveEditorSelection.h"
+#include "CurveEditorSelectionPrivate.h"
+#include "Modification/Selection/SelectionDeltaChange.h"
+
+#include <type_traits>
+
+namespace UE::CurveEditor
+{
+namespace SelectionDiffDetail
+{
+/** Goes through every key in InOriginalSelection and sees which keys are missing from InTarget. The missing keys are added to the array InGetKeys returns. */
+template<typename TGetKeyArray>
+requires std::is_invocable_r_v<TArray<TPair<FKeyHandle, ECurvePointType>>&, TGetKeyArray, FCurveSelectionDeltaChange&>
+static void DetectMissingKeys(
+	const TMap<FCurveModelID, FKeyHandleSet>& InOriginalSelection, const FCurveEditorSelection& InTarget,
+	FSelectionDeltaChange& OutDeltaChange, TGetKeyArray&& InGetKeys
+	)
+{
+	for (const TPair<FCurveModelID, FKeyHandleSet>& OriginalCurvePair : InOriginalSelection)
+	{
+		const FCurveModelID& CurveId = OriginalCurvePair.Key;
+		const FKeyHandleSet& OriginalKeySet = OriginalCurvePair.Value; 
+		const FKeyHandleSet* TargetKeySet = InTarget.FindForCurve(CurveId);
+
+		// Should not really happen, but we'll handle the case anyway.
+		if (OriginalKeySet.Num() == 0 && (!TargetKeySet || TargetKeySet->Num() == 0))
+		{
+			continue;
+		}	
+
+		// Fully missing from target?
+		if (!TargetKeySet
+			// We'd expect that FindForCurve would have returned nullptr if Num() == 0 but we'll handle the case anyway.
+			|| TargetKeySet->Num() == 0)
+		{
+			FCurveSelectionDeltaChange& CurveChange = OutDeltaChange.ChangedCurves.Add(CurveId);
+			TArray<TPair<FKeyHandle, ECurvePointType>>& Keys = InGetKeys(CurveChange);
+			Keys.Reserve(OriginalKeySet.Num());
+			for (const FKeyHandle& RemovedKey : OriginalKeySet.AsArray())
+			{
+				Keys.Emplace(RemovedKey, OriginalKeySet.PointType(RemovedKey));
+			}
+			continue;
+		}
+
+		// Find keys that were are missing or changed point type
+		for (const FKeyHandle& OriginalKey : OriginalKeySet.AsArray())
+		{
+			// If the key handle is still in there but the point type has changed, we'll add it back again below.
+			const ECurvePointType PointType = OriginalKeySet.PointType(OriginalKey);
+			if (!TargetKeySet->Contains(OriginalKey, PointType))
+			{
+				TArray<TPair<FKeyHandle, ECurvePointType>>& Keys = InGetKeys(OutDeltaChange.ChangedCurves.FindOrAdd(CurveId));
+				Keys.Emplace(OriginalKey, PointType);
+			}
+		}
+	}
+}
+	
+template<typename TCallback>
+concept CGetKeys = std::is_invocable_r_v<const TArray<TPair<FKeyHandle, ECurvePointType>>&, TCallback, const FCurveSelectionDeltaChange&>;
+
+void AddKeysInternal(FCurveEditorSelection& Selection, FCurveModelID CurveID, ECurvePointType PointType, TArrayView<const FKeyHandle> Keys)
+{
+	KeySelection::FAddInternal(Selection, CurveID, PointType, Keys);
+}
+
+void SetSerialNumberInternal(FCurveEditorSelection& Selection, uint32 NewSerialNumber)
+{
+	KeySelection::FSetSerialNumber(Selection, NewSerialNumber);
+}
+	
+/** Removes keys from the array InGetRemovedKeys returns, and adds keys to the array InGetAddedKeys returns. */
+template<CGetKeys TGetAddedKeys, CGetKeys TGetRemovedKeys>
+static void ModifySelection(
+	FCurveEditorSelection& InOutChange, const FSelectionDeltaChange& InDeltaChange, TGetAddedKeys&& InGetAddedKeys, TGetRemovedKeys&& InGetRemovedKeys
+	)
+{
+	for (const TPair<FCurveModelID, FCurveSelectionDeltaChange>& ChangePair : InDeltaChange.ChangedCurves)
+	{
+		const FCurveModelID& CurveId = ChangePair.Key;
+		const TArray<TPair<FKeyHandle, ECurvePointType>>& AddedKeys = InGetAddedKeys(ChangePair.Value);
+		const TArray<TPair<FKeyHandle, ECurvePointType>>& RemovedKeys = InGetRemovedKeys(ChangePair.Value);
+	
+		for (const TPair<FKeyHandle, ECurvePointType>& RemovedPair : RemovedKeys)
+		{
+			InOutChange.Remove(CurveId, RemovedPair.Value, RemovedPair.Key);
+		}
+		for (const TPair<FKeyHandle, ECurvePointType>& AddedPair : AddedKeys)
+		{
+			const TConstArrayView<FKeyHandle> Keys(&AddedPair.Key, 1);
+			AddKeysInternal(InOutChange, CurveId, AddedPair.Value, Keys);
+		}
+	}
+}
+}
+	
+FSelectionDeltaChange DiffSelection(const FCurveEditorSelection& InOriginal, const FCurveEditorSelection& InTarget)
+{
+	FSelectionDeltaChange DeltaChange;
+	const TMap<FCurveModelID, FKeyHandleSet>& OriginalSelection = InOriginal.GetAll();
+	const TMap<FCurveModelID, FKeyHandleSet>& TargetSelection = InTarget.GetAll();
+	
+	// Go through everything in InOriginal and see what's missing in InTarget -> effectively computes what was removed.
+	SelectionDiffDetail::DetectMissingKeys(OriginalSelection, InTarget, DeltaChange,
+		[](FCurveSelectionDeltaChange& CurveChanges) -> TArray<TPair<FKeyHandle, ECurvePointType>>&
+		{ return CurveChanges.RemovedKeys; }
+		);
+	// Go through everything in InTarget and see what's missing in InOriginal -> effectively computes what was added.
+	SelectionDiffDetail::DetectMissingKeys(TargetSelection, InOriginal, DeltaChange,
+		[](FCurveSelectionDeltaChange& CurveChanges) -> TArray<TPair<FKeyHandle, ECurvePointType>>&
+		{ return CurveChanges.AddedKeys; }
+	);
+
+	DeltaChange.OldSerialNumber = InOriginal.GetSerialNumber();
+	DeltaChange.NewSerialNumber = InTarget.GetSerialNumber();
+	
+	return DeltaChange;
+}
+
+void ApplySelectionChange(FCurveEditorSelection& InOutChange, const FSelectionDeltaChange& InDeltaChange)
+{
+	SelectionDiffDetail::ModifySelection(InOutChange, InDeltaChange,
+		[](const FCurveSelectionDeltaChange& CurveChange){ return CurveChange.AddedKeys; },
+		[](const FCurveSelectionDeltaChange& CurveChange){ return CurveChange.RemovedKeys; }
+		);
+
+	SelectionDiffDetail::SetSerialNumberInternal(InOutChange, InDeltaChange.NewSerialNumber);
+}
+	
+void RevertSelectionChange(FCurveEditorSelection& InOutChange, const FSelectionDeltaChange& InDeltaChange)
+{
+	SelectionDiffDetail::ModifySelection(InOutChange, InDeltaChange,
+		// Swapping added and removed keys here effectively reverts the change.
+		[](const FCurveSelectionDeltaChange& CurveChange){ return CurveChange.RemovedKeys; },
+		[](const FCurveSelectionDeltaChange& CurveChange){ return CurveChange.AddedKeys; }
+		);
+	
+	SelectionDiffDetail::SetSerialNumberInternal(InOutChange, InDeltaChange.OldSerialNumber);
+}
+}
diff --git a/Engine/Source/Editor/CurveEditor/Private/Modification/Selection/SelectionChangeUtils.h b/Engine/Source/Editor/CurveEditor/Private/Modification/Selection/SelectionChangeUtils.h
new file mode 100644
--- /dev/null
+++ b/Engine/Source/Editor/CurveEditor/Private/Modification/Selection/SelectionChangeUtils.h
@@ -0,0 +1,27 @@
+﻿// Copyright Epic Games, Inc. All Rights Reserved.
+
+#pragma once
+
+#include "Containers/ArrayView.h"
+
+enum class ECurvePointType : uint8;
+struct FCurveModelID;
+struct FCurveEditorSelection;
+struct FKeyHandle;
+
+namespace UE::CurveEditor
+{
+struct FSelectionDeltaChange;
+
+/** Computes the delta change to get from InOriginal to InTarget. */
+FSelectionDeltaChange DiffSelection(const FCurveEditorSelection& InOriginal, const FCurveEditorSelection& InTarget);
+
+/** Applies InDeltaChange to InOutSelection (redo operation). */
+void ApplySelectionChange(FCurveEditorSelection& InOutChange, const FSelectionDeltaChange& InDeltaChange);
+/** Reverts InDeltaChange from InOutSelection (undo operation).  */
+void RevertSelectionChange(FCurveEditorSelection& InOutChange, const FSelectionDeltaChange& InDeltaChange);
+
+// Private - needed for FAddInternal.
+namespace SelectionDiffDetail { void AddKeysInternal(FCurveEditorSelection&, FCurveModelID, ECurvePointType, TArrayView<const FKeyHandle>); }
+namespace SelectionDiffDetail { void SetSerialNumberInternal(FCurveEditorSelection& Selection, uint32 NewSerialNumber); }
+}
diff --git a/Engine/Source/Editor/CurveEditor/Private/Modification/Selection/SelectionDeltaChange.h b/Engine/Source/Editor/CurveEditor/Private/Modification/Selection/SelectionDeltaChange.h
new file mode 100644
--- /dev/null
+++ b/Engine/Source/Editor/CurveEditor/Private/Modification/Selection/SelectionDeltaChange.h
@@ -0,0 +1,40 @@
+﻿// Copyright Epic Games, Inc. All Rights Reserved.
+
+#pragma once
+
+#include "Containers/Array.h"
+#include "Containers/Map.h"
+#include "CurveEditorTypes.h"
+#include "CurveModel.h"
+
+namespace UE::CurveEditor
+{
+/** Changes made to a single curve's selection */
+struct FCurveSelectionDeltaChange
+{
+	// Using TArray instead of TMap because we'll be iterating through all entries frequently. TArray is much faster for that.
+	
+	/** The keys that were added and its point type. */
+	TArray<TPair<FKeyHandle, ECurvePointType>> AddedKeys;
+	/**
+	 * The keys that were removed.
+	 * Can contain a key handle that is in AddedKeys. In that case, the point type was changed.
+	 */
+	TArray<TPair<FKeyHandle, ECurvePointType>> RemovedKeys;
+};
+
+/** Changes made to the curve editor selection */
+struct FSelectionDeltaChange
+{
+	/** The per curve changes. */
+	TMap<FCurveModelID, FCurveSelectionDeltaChange> ChangedCurves;
+
+	/** Serial number of selection before the change. */
+	uint32 OldSerialNumber = 0;
+	/** Serial number of selection after the change. */
+	uint32 NewSerialNumber = 0;
+
+	bool HasChanges() const { return !ChangedCurves.IsEmpty() || OldSerialNumber != NewSerialNumber; }
+	operator bool() const { return HasChanges(); }
+};
+}
diff --git a/Engine/Source/Editor/CurveEditor/Private/Modification/TransactionManager.cpp b/Engine/Source/Editor/CurveEditor/Private/Modification/TransactionManager.cpp
--- a/Engine/Source/Editor/CurveEditor/Private/Modification/TransactionManager.cpp
+++ b/Engine/Source/Editor/CurveEditor/Private/Modification/TransactionManager.cpp
@@ -3,6 +3,7 @@
 #include "Modification/TransactionManager.h"
 
 #include "CurveEditorTransactionObject.h"
+#include "ScopedTransaction.h"
 #include "Misc/ITransaction.h"
 #include "UObject/Package.h"
 #include "UObject/UObjectGlobals.h"
@@ -15,9 +16,11 @@ FTransactionManager::FTransactionManager(TWeakPtr<FCurveEditor> InCurveEditor)
 	DummyTransactionObject->OwningCurveEditor = MoveTemp(InCurveEditor);
 }
 
-void FTransactionManager::AppendChange(TUniquePtr<FCurveEditorCommandChange> InChange) const
+void FTransactionManager::AppendChange(TUniquePtr<FCurveEditorCommandChange> InChange, const FText& InDescription) const
 {
-	if (ensureAlways(GUndo) && ensure(DummyTransactionObject) && ensure(InChange))
+	FScopedTransaction Transaction(InDescription);
+	if (GUndo // GUndo can be reset by external systems, e.g. during world tear down.
+		&& ensure(DummyTransactionObject) && ensure(InChange))
 	{
 		GUndo->StoreUndo(DummyTransactionObject, MoveTemp(InChange));
 		OnCommandAppendedDelegate.Broadcast(*InChange);
diff --git a/Engine/Source/Editor/CurveEditor/Private/Modification/Utils/ScopedChangeBase.cpp b/Engine/Source/Editor/CurveEditor/Private/Modification/Utils/ScopedChangeBase.cpp
new file mode 100644
--- /dev/null
+++ b/Engine/Source/Editor/CurveEditor/Private/Modification/Utils/ScopedChangeBase.cpp
@@ -0,0 +1,29 @@
+﻿// Copyright Epic Games, Inc. All Rights Reserved.
+
+#include "Modification/Utils/ScopedChangeBase.h"
+
+#include "CoreGlobals.h"
+#include "CurveEditor.h"
+
+namespace UE::CurveEditor
+{
+FScopedChangeBase::FScopedChangeBase(TWeakPtr<FCurveEditor> InCurveEditor, const FText& InDescription)
+	: WeakCurveEditor(MoveTemp(InCurveEditor))
+	, Description(InDescription)
+	, bWasStartedInTransaction(GUndo != nullptr)
+{}
+
+bool FScopedChangeBase::IsCancelled() const
+{
+	return bIsCancelled || (GUndo == nullptr && bWasStartedInTransaction);
+}
+
+void FScopedChangeBase::TryAppendCommand(TUniquePtr<FCurveEditorCommandChange> InCommand) const
+{
+	const TSharedPtr<FCurveEditor> CurveEditorPin = WeakCurveEditor.Pin();
+	if (CurveEditorPin && !IsCancelled())
+	{
+		CurveEditorPin->GetTransactionManager()->AppendChange(MoveTemp(InCommand), Description);
+	}
+}
+}
diff --git a/Engine/Source/Editor/CurveEditor/Private/Modification/Utils/ScopedSelectionChange.cpp b/Engine/Source/Editor/CurveEditor/Private/Modification/Utils/ScopedSelectionChange.cpp
new file mode 100644
--- /dev/null
+++ b/Engine/Source/Editor/CurveEditor/Private/Modification/Utils/ScopedSelectionChange.cpp
@@ -0,0 +1,42 @@
+﻿// Copyright Epic Games, Inc. All Rights Reserved.
+
+#include "Modification/Utils/ScopedSelectionChange.h"
+
+#include "CurveEditor.h"
+#include "Internationalization/Internationalization.h"
+#include "Modification/Selection/SelectionChangeCommand.h"
+#include "Modification/Selection/SelectionChangeUtils.h"
+
+#define LOCTEXT_NAMESPACE "FScopedSelectionChange"
+
+namespace UE::CurveEditor
+{
+FScopedSelectionChange::FScopedSelectionChange(TWeakPtr<FCurveEditor> InCurveEditor)
+	: FScopedSelectionChange(MoveTemp(InCurveEditor), LOCTEXT("ChangeSelection", "Change selection"))
+{}
+
+FScopedSelectionChange::FScopedSelectionChange(TWeakPtr<FCurveEditor> InCurveEditor, const FText& InDescription)
+	: FScopedChangeBase(MoveTemp(InCurveEditor), InDescription)
+{
+	if (const TSharedPtr<FCurveEditor> CurveEditorPin = WeakCurveEditor.Pin())
+	{
+		OriginalSelection = CurveEditorPin->Selection;
+	}
+}
+
+FScopedSelectionChange::~FScopedSelectionChange()
+{
+	const TSharedPtr<FCurveEditor> CurveEditorPin = WeakCurveEditor.Pin();
+	if (IsCancelled() || !CurveEditorPin)
+	{
+		return;
+	}
+	
+	if (FSelectionDeltaChange DeltaChange = DiffSelection(OriginalSelection, CurveEditorPin->Selection))
+	{
+		TryAppendCommand(MakeUnique<FSelectionChangeCommand>(MoveTemp(DeltaChange)));
+	}
+}
+}
+
+#undef LOCTEXT_NAMESPACE
\ No newline at end of file
diff --git a/Engine/Source/Editor/CurveEditor/Private/Modification/Utils/ScopedSelectionTransaction.cpp b/Engine/Source/Editor/CurveEditor/Private/Modification/Utils/ScopedSelectionTransaction.cpp
deleted file mode 100644
diff --git a/Engine/Source/Editor/CurveEditor/Private/Modification/Utils/SelectionChangeUtils.cpp b/Engine/Source/Editor/CurveEditor/Private/Modification/Utils/SelectionChangeUtils.cpp
deleted file mode 100644
diff --git a/Engine/Source/Editor/CurveEditor/Private/Modification/Utils/SelectionChangeUtils.h b/Engine/Source/Editor/CurveEditor/Private/Modification/Utils/SelectionChangeUtils.h
deleted file mode 100644
diff --git a/Engine/Source/Editor/CurveEditor/Private/SCurveEditorPanel.cpp b/Engine/Source/Editor/CurveEditor/Private/SCurveEditorPanel.cpp
--- a/Engine/Source/Editor/CurveEditor/Private/SCurveEditorPanel.cpp
+++ b/Engine/Source/Editor/CurveEditor/Private/SCurveEditorPanel.cpp
@@ -67,7 +67,7 @@
 #include "Filters/CurveEditorEulerFilter.h"
 #include "Filters/PromotedFilterContainer.h"
 #include "Misc/ResizeParamUtils.h"
-#include "Modification/Utils/ScopedSelectionTransaction.h"
+#include "Modification/Utils/ScopedSelectionChange.h"
 #include "Styling/AppStyle.h"
 #include "Styling/CoreStyle.h"
 #include "Styling/ISlateStyle.h"
@@ -461,7 +461,7 @@ void SCurveEditorPanel::BindCommands()
 		TSharedPtr<FCurveEditor> LocalCurveEditor = CurveEditor;
 		FExecuteAction DeselectAllAction = FExecuteAction::CreateLambda([LocalCurveEditor] { if (LocalCurveEditor.IsValid())
 			{
-				const UE::CurveEditor::FScopedSelectionTransaction Transaction(LocalCurveEditor, LOCTEXT("DeselectAllKeys", "Deselect all keys"));
+				const UE::CurveEditor::FScopedSelectionChange Transaction(LocalCurveEditor, LOCTEXT("DeselectAllKeys", "Deselect all keys"));
 				LocalCurveEditor->GetSelection().Clear();
 			} 
 		});
@@ -1139,7 +1139,7 @@ FReply SCurveEditorPanel::OnKeyDown(const FGeometry& MyGeometry, const FKeyEvent
 {
 	if (InKeyEvent.GetKey() == EKeys::Escape)
 	{
-		const UE::CurveEditor::FScopedSelectionTransaction Transaction(GetCurveEditor(), LOCTEXT("EscapeKeys", "Press escape"));
+		const UE::CurveEditor::FScopedSelectionChange Transaction(GetCurveEditor(), LOCTEXT("EscapeKeys", "Press escape"));
 		CurveEditor->Selection.Clear();
 		return FReply::Handled();
 	}
diff --git a/Engine/Source/Editor/CurveEditor/Private/SCurveEditorViewContainer.cpp b/Engine/Source/Editor/CurveEditor/Private/SCurveEditorViewContainer.cpp
--- a/Engine/Source/Editor/CurveEditor/Private/SCurveEditorViewContainer.cpp
+++ b/Engine/Source/Editor/CurveEditor/Private/SCurveEditorViewContainer.cpp
@@ -370,7 +370,7 @@ FReply SCurveEditorViewContainer::OnMouseButtonUp(const FGeometry& MyGeometry, c
 	{
 		if (!MouseEvent.IsShiftDown() && !MouseEvent.IsAltDown() && !MouseEvent.IsControlDown())
 		{
-			const UE::CurveEditor::FScopedSelectionTransaction Transaction(CurveEditor, LOCTEXT("ClearSelection", "Clear selection"));
+			const UE::CurveEditor::FScopedSelectionChange Transaction(CurveEditor, LOCTEXT("ClearSelection", "Clear selection"));
 			CurveEditor->GetSelection().Clear();
 			Reply = FReply::Handled();
 		}
diff --git a/Engine/Source/Editor/CurveEditor/Private/Tree/SCurveEditorTreeSelect.cpp b/Engine/Source/Editor/CurveEditor/Private/Tree/SCurveEditorTreeSelect.cpp
--- a/Engine/Source/Editor/CurveEditor/Private/Tree/SCurveEditorTreeSelect.cpp
+++ b/Engine/Source/Editor/CurveEditor/Private/Tree/SCurveEditorTreeSelect.cpp
@@ -13,7 +13,7 @@
 #include "GenericPlatform/GenericApplication.h"
 #include "Layout/Children.h"
 #include "Math/NumericLimits.h"
-#include "Modification/Utils/ScopedSelectionTransaction.h"
+#include "Modification/Utils/ScopedSelectionChange.h"
 #include "Styling/AppStyle.h"
 #include "Templates/UnrealTemplate.h"
 #include "Tree/CurveEditorTree.h"
@@ -61,7 +61,7 @@ FReply SCurveEditorTreeSelect::SelectAll()
 	const bool bIsShiftDown = FSlateApplication::Get().GetModifierKeys().IsShiftDown();
 	const bool bIsControlDown = FSlateApplication::Get().GetModifierKeys().IsControlDown();
 
-	const UE::CurveEditor::FScopedSelectionTransaction Transaction(CurveEditor, LOCTEXT("Tree.SelectAll", "Select all"));
+	const UE::CurveEditor::FScopedSelectionChange Transaction(CurveEditor, LOCTEXT("Tree.SelectAll", "Select all"));
 	if (!bIsShiftDown && !bIsControlDown)
 	{
 		CurveEditor->GetSelection().Clear();
diff --git a/Engine/Source/Editor/CurveEditor/Private/Views/SInteractiveCurveEditorView.cpp b/Engine/Source/Editor/CurveEditor/Private/Views/SInteractiveCurveEditorView.cpp
--- a/Engine/Source/Editor/CurveEditor/Private/Views/SInteractiveCurveEditorView.cpp
+++ b/Engine/Source/Editor/CurveEditor/Private/Views/SInteractiveCurveEditorView.cpp
@@ -63,7 +63,7 @@
 #include "SCurveEditorPanel.h"
 #include "ScopedTransaction.h"
 #include "SlotBase.h"
-#include "Modification/Utils/ScopedSelectionTransaction.h"
+#include "Modification/Utils/ScopedSelectionChange.h"
 #include "Styling/AppStyle.h"
 #include "Styling/CoreStyle.h"
 #include "Styling/ISlateStyle.h"
@@ -1470,7 +1470,7 @@ double SInteractiveCurveEditorView::GetTangentValue(const double InTime, const d
 
 void SInteractiveCurveEditorView::HandleDirectKeySelectionByMouse(TSharedPtr<FCurveEditor> CurveEditor, const FPointerEvent& MouseEvent, TOptional<FCurvePointHandle> MouseDownPoint)
 {
-	const UE::CurveEditor::FScopedSelectionTransaction Transaction(WeakCurveEditor, LOCTEXT("ClickKey", "Click Key"));
+	const UE::CurveEditor::FScopedSelectionChange Transaction(WeakCurveEditor, LOCTEXT("ClickKey", "Click Key"));
 	if(!MouseDownPoint.IsSet())
 	{
 		CurveEditor->GetSelection().Clear();
@@ -1634,7 +1634,7 @@ FReply SInteractiveCurveEditorView::OnMouseButtonDown(const FGeometry& MyGeometr
 				FCurveModel* CurveToAddTo = CurveEditor->FindCurve(HoveredCurve.GetValue());
 				if (CurveToAddTo && !CurveToAddTo->IsReadOnly() && !bDraggingKeys && !bKeyAlreadyExists)
 				{
-					UE::CurveEditor::FScopedSelectionTransaction Transaction(CurveEditor, LOCTEXT("InsertKey", "Insert Key"));
+					UE::CurveEditor::FScopedSelectionChange Transaction(CurveEditor, LOCTEXT("InsertKey", "Insert Key"));
 
 					FCurveEditorScreenSpace CurveSpace = GetCurveSpace(HoveredCurve.GetValue());
 					double MouseTime = CurveSpace.ScreenToSeconds(MousePixel.X);
@@ -1802,7 +1802,7 @@ FReply SInteractiveCurveEditorView::OnMouseButtonUp(const FGeometry& MyGeometry,
 		TOptional<FCurveModelID> HitCurve = GetHoveredCurve();
 		if (!HitPoint(MousePosition).IsSet() && HitCurve.IsSet())
 		{
-			const UE::CurveEditor::FScopedSelectionTransaction Transaction(WeakCurveEditor, LOCTEXT("ClickCurve", "Click Curve"));
+			const UE::CurveEditor::FScopedSelectionChange Transaction(WeakCurveEditor, LOCTEXT("ClickCurve", "Click Curve"));
 			FCurveModel* CurveModel = CurveEditor->FindCurve(HitCurve.GetValue());
 
 			TArray<FKeyHandle> KeyHandles;
@@ -2112,7 +2112,7 @@ void SInteractiveCurveEditorView::AddKeyAtTime(const TSet<FCurveModelID>& ToCurv
 		return;
 	}
 
-	UE::CurveEditor::FScopedSelectionTransaction Transaction(WeakCurveEditor, LOCTEXT("AddKeyAtTime", "Add Key"));
+	UE::CurveEditor::FScopedSelectionChange Transaction(WeakCurveEditor, LOCTEXT("AddKeyAtTime", "Add Key"));
 	bool bAddedKey = false;
 
 	// Clear the selection set as we will be selecting all the new keys created.
diff --git a/Engine/Source/Editor/CurveEditor/Public/Modification/TransactionManager.h b/Engine/Source/Editor/CurveEditor/Public/Modification/TransactionManager.h
--- a/Engine/Source/Editor/CurveEditor/Public/Modification/TransactionManager.h
+++ b/Engine/Source/Editor/CurveEditor/Public/Modification/TransactionManager.h
@@ -29,8 +29,9 @@ public:
 	 * - system 2 could coniditonally snap the keys to full frames after checking that the user has toggled auto-snapping
 	 * 
 	 * @param InChange The change to associate with the currently open transaction
+	 * @param InDescription Description of the change
 	 */
-	UE_API void AppendChange(TUniquePtr<FCurveEditorCommandChange> InChange) const;
+	UE_API void AppendChange(TUniquePtr<FCurveEditorCommandChange> InChange, const FText& InDescription) const;
 
 	DECLARE_MULTICAST_DELEGATE_OneParam(FOnCommandAppended, const FCurveEditorCommandChange&);
 	/** Invoked after a command has been appended to the undo stack. */
diff --git a/Engine/Source/Editor/CurveEditor/Public/Modification/Utils/ScopedChangeBase.h b/Engine/Source/Editor/CurveEditor/Public/Modification/Utils/ScopedChangeBase.h
new file mode 100644
--- /dev/null
+++ b/Engine/Source/Editor/CurveEditor/Public/Modification/Utils/ScopedChangeBase.h
@@ -0,0 +1,45 @@
+﻿// Copyright Epic Games, Inc. All Rights Reserved.
+
+#pragma once
+
+#include "CurveEditorSelection.h"
+#include "ScopedTransaction.h"
+
+namespace UE::CurveEditor
+{
+class FCurveEditorCommandChange;
+
+/**
+ * Base class for scoped changes that intend to append a command to the transaction system when the scope ends.
+ * It detects if constructed within a FScopedTransaction context: by calling IsCancelled(), you can handle the case of the root transaction being cancelled.
+ */
+class FScopedChangeBase
+{
+public:
+
+	[[nodiscard]] explicit FScopedChangeBase(TWeakPtr<FCurveEditor> InCurveEditor, const FText& InDescription);
+
+	/** Prevent the change from being appended. */
+	void Cancel() { bIsCancelled = true; }
+
+protected:
+	
+	/** Used to diff selection and append the change to the transaction system. */
+	const TWeakPtr<FCurveEditor> WeakCurveEditor;
+	/** Description to use for the change */
+	const FText Description;
+	
+	/** Whether the change should be appended. */
+	bool IsCancelled() const;
+
+	/** Appends the change if it is not cancelled. */
+	void TryAppendCommand(TUniquePtr<FCurveEditorCommandChange> InCommand) const;
+	
+private:
+
+	/** Whether we were constructed in a transaction. Used to detect whether the transaction has since been cancelled. */
+	const bool bWasStartedInTransaction;
+	/** Whether to skip submitting the change when the scope ends. */
+	bool bIsCancelled = false;
+};
+}
diff --git a/Engine/Source/Editor/CurveEditor/Public/Modification/Utils/ScopedSelectionChange.h b/Engine/Source/Editor/CurveEditor/Public/Modification/Utils/ScopedSelectionChange.h
new file mode 100644
--- /dev/null
+++ b/Engine/Source/Editor/CurveEditor/Public/Modification/Utils/ScopedSelectionChange.h
@@ -0,0 +1,35 @@
+﻿// Copyright Epic Games, Inc. All Rights Reserved.
+
+#pragma once
+
+#include "CurveEditorSelection.h"
+#include "ScopedChangeBase.h"
+
+#define UE_API CURVEEDITOR_API
+
+namespace UE::CurveEditor
+{
+/**
+ * Saves the current selection state and on destruction compares it against the then current selection. If the selection has changed, it appends
+ * an undo-able action to the curve editor transaction manager.
+ *
+ * You can safely use this within the scope of a FScopedTransaction. In that case, the command will be added as sub-transaction to the parent
+ * transaction.
+ */
+class FScopedSelectionChange : public FScopedChangeBase
+{
+public:
+
+	[[nodiscard]] UE_API explicit FScopedSelectionChange(TWeakPtr<FCurveEditor> InCurveEditor);
+	[[nodiscard]] UE_API explicit FScopedSelectionChange(TWeakPtr<FCurveEditor> InCurveEditor, const FText& InDescription);
+	
+	UE_API ~FScopedSelectionChange();
+
+private:
+
+	/** The selection the editor had when the transaction was started.  */
+	FCurveEditorSelection OriginalSelection;
+};
+}
+
+#undef UE_API
diff --git a/Engine/Source/Editor/CurveEditor/Public/Modification/Utils/ScopedSelectionTransaction.h b/Engine/Source/Editor/CurveEditor/Public/Modification/Utils/ScopedSelectionTransaction.h
deleted file mode 100644
diff --git a/Engine/Source/Editor/Sequencer/Private/MVVM/CurveEditorExtension.cpp b/Engine/Source/Editor/Sequencer/Private/MVVM/CurveEditorExtension.cpp
--- a/Engine/Source/Editor/Sequencer/Private/MVVM/CurveEditorExtension.cpp
+++ b/Engine/Source/Editor/Sequencer/Private/MVVM/CurveEditorExtension.cpp
@@ -20,7 +20,7 @@
 #include "Sequencer.h"
 #include "SequencerCommands.h"
 #include "Menus/SequencerToolbarUtils.h"
-#include "Modification/Utils/ScopedSelectionTransaction.h"
+#include "Modification/Utils/ScopedSelectionChange.h"
 #include "Toolkits/IToolkitHost.h"
 #include "Tree/SCurveEditorTree.h"
 #include "Tree/SCurveEditorTreeFilterStatusBar.h"
@@ -627,9 +627,15 @@ void FCurveEditorExtension::RequestSyncSelection()
 
 	// Key selection supports undo. If RequestSyncSelection is called as part of an ongoing transaction, record the key selection change for undo.
 	const bool bShouldRecord = GUndo != nullptr;
-	TSharedPtr<CurveEditor::FScopedSelectionTransaction> SelectionChange = MakeShared<CurveEditor::FScopedSelectionTransaction>(CurveEditorModel, bShouldRecord);
+	// The current transaction needs to be extended until next tick...
+	TSharedPtr<FScopedTransaction> Transaction = bShouldRecord
+		? MakeShared<FScopedTransaction>(FText::GetEmpty(), bShouldRecord).ToSharedPtr() : nullptr;
+	// ... at which point we'll diff the changes.
+	TSharedPtr<CurveEditor::FScopedSelectionChange> SelectionChange = bShouldRecord
+		? MakeShared<CurveEditor::FScopedSelectionChange>(CurveEditorModel).ToSharedPtr() : nullptr;
 	
-	GEditor->GetTimerManager()->SetTimerForNextTick([WeakRootViewModel, SelectionChange = MoveTemp(SelectionChange)]() mutable
+	GEditor->GetTimerManager()->SetTimerForNextTick(
+	[WeakRootViewModel, Transaction = MoveTemp(Transaction), SelectionChange = MoveTemp(SelectionChange)]() mutable
 	{
 		bSyncSelectionRequested = false;
 	
@@ -650,7 +656,10 @@ void FCurveEditorExtension::RequestSyncSelection()
 			This->SyncSelection();
 		}
 
+		// Order matters: this appends a sub-transaction for the selection change...
 		SelectionChange.Reset();
+		// ... and only then can the parent transaction be closed. 
+		Transaction.Reset();
 	});
 }
 
