Curve Editor: Fix undo / redo not working for the following operations:
- Paste keys
- Delete keys
- Insert key (middle-mouse click curve)
- Add key (right-click curve -> add)

diff --git a/Engine/Source/Editor/CurveEditor/Private/CurveEditor.cpp b/Engine/Source/Editor/CurveEditor/Private/CurveEditor.cpp
--- a/Engine/Source/Editor/CurveEditor/Private/CurveEditor.cpp
+++ b/Engine/Source/Editor/CurveEditor/Private/CurveEditor.cpp
@@ -1845,8 +1845,9 @@ void FCurveEditor::PasteKeys(UE::CurveEditor::FKeyPasteArgs InArgs)
 			TimeOffset = SourceBuffer->TimeOffset;
 		}
 	}
-
-	const UE::CurveEditor::FScopedSelectionChange KeyChange(SharedThis(this), LOCTEXT("PasteKeys", "Paste Keys"));
+	
+	const FScopedTransaction Transaction(LOCTEXT("PasteKeys", "Paste Keys")); // Selection and key change should be part of same undo / redo op.
+	const UE::CurveEditor::FScopedSelectionChange KeyChange(SharedThis(this));
 	Selection.Clear();
 
 
@@ -1976,7 +1977,8 @@ void FCurveEditor::PasteKeys(UE::CurveEditor::FKeyPasteArgs InArgs)
 
 void FCurveEditor::DeleteSelection()
 {
-	const UE::CurveEditor::FScopedSelectionChange Transaction(SharedThis(this), LOCTEXT("DeleteKeys", "Delete Keys"));
+	const FScopedTransaction Transaction(LOCTEXT("DeleteKeys", "Delete Keys")); // Selection and key change should be part of same undo / redo op.
+	const UE::CurveEditor::FScopedSelectionChange SelectionChange(SharedThis(this));
 
 	double CurrentTime = 0.0;
 	if (TSharedPtr<ITimeSliderController> TimeSliderController = WeakTimeSliderController.Pin())
@@ -2143,7 +2145,8 @@ bool FCurveEditor::CanFlattenOrStraightenSelection() const
 
 void FCurveEditor::SmartSnapSelection()
 {
-	const UE::CurveEditor::FScopedSelectionChange Transaction(SharedThis(this), LOCTEXT("SmartSnapKeys", "Smart Snap"));
+	const FScopedTransaction Transaction(LOCTEXT("SmartSnapKeys", "Smart Snap")); // Selection and key change should be part of same undo / redo op.
+	const UE::CurveEditor::FScopedSelectionChange SelectionChange(SharedThis(this));
 	
 	TMap<FCurveModelID, FKeyHandleSet> OutKeysToSelect;
 	UE::CurveEditor::EnumerateSmartSnappableKeys(*this, Selection.GetAll(), OutKeysToSelect,
diff --git a/Engine/Source/Editor/CurveEditor/Private/Views/SInteractiveCurveEditorView.cpp b/Engine/Source/Editor/CurveEditor/Private/Views/SInteractiveCurveEditorView.cpp
--- a/Engine/Source/Editor/CurveEditor/Private/Views/SInteractiveCurveEditorView.cpp
+++ b/Engine/Source/Editor/CurveEditor/Private/Views/SInteractiveCurveEditorView.cpp
@@ -1634,7 +1634,8 @@ FReply SInteractiveCurveEditorView::OnMouseButtonDown(const FGeometry& MyGeometr
 				FCurveModel* CurveToAddTo = CurveEditor->FindCurve(HoveredCurve.GetValue());
 				if (CurveToAddTo && !CurveToAddTo->IsReadOnly() && !bDraggingKeys && !bKeyAlreadyExists)
 				{
-					UE::CurveEditor::FScopedSelectionChange Transaction(CurveEditor, LOCTEXT("InsertKey", "Insert Key"));
+					FScopedTransaction Transaction(LOCTEXT("InsertKey", "Insert Key")); // Selection and key change should be part of same undo / redo op.
+					UE::CurveEditor::FScopedSelectionChange SelectionChange(CurveEditor);
 
 					FCurveEditorScreenSpace CurveSpace = GetCurveSpace(HoveredCurve.GetValue());
 					double MouseTime = CurveSpace.ScreenToSeconds(MousePixel.X);
@@ -2112,7 +2113,8 @@ void SInteractiveCurveEditorView::AddKeyAtTime(const TSet<FCurveModelID>& ToCurv
 		return;
 	}
 
-	UE::CurveEditor::FScopedSelectionChange Transaction(WeakCurveEditor, LOCTEXT("AddKeyAtTime", "Add Key"));
+	FScopedTransaction Transaction(LOCTEXT("AddKeyAtTime", "Add Key")); // Selection and key change should be part of same undo / redo op.
+	UE::CurveEditor::FScopedSelectionChange SelectionChange(WeakCurveEditor);
 	bool bAddedKey = false;
 
 	// Clear the selection set as we will be selecting all the new keys created.
