Curve Editor: Check all chord combinations for scrubbing time. Otherwise, if you bind to c, Ctrl-C event is eaten and the user cannot copy keys

Same fix as 28715772

diff --git a/Engine/Source/Editor/CurveEditor/Private/SCurveEditorViewContainer.cpp b/Engine/Source/Editor/CurveEditor/Private/SCurveEditorViewContainer.cpp
--- a/Engine/Source/Editor/CurveEditor/Private/SCurveEditorViewContainer.cpp
+++ b/Engine/Source/Editor/CurveEditor/Private/SCurveEditorViewContainer.cpp
@@ -169,7 +169,11 @@ bool SCurveEditorViewContainer::IsScrubTimeKeyEvent(const FKeyEvent& InKeyEvent)
 	{
 		EMultipleKeyBindingIndex ChordIndex = static_cast<EMultipleKeyBindingIndex>(i);
 		const FInputChord& Chord = *Commands.ScrubTime->GetActiveChord(ChordIndex);
-		const bool bIsMovingTimeSlider = Chord.IsValidChord() && InKeyEvent.GetKey() == Chord.Key;
+		const bool bIsMovingTimeSlider = Chord.IsValidChord() && InKeyEvent.GetKey() == Chord.Key &&
+			(Chord.NeedsAlt() ? InKeyEvent.IsAltDown() : !InKeyEvent.IsAltDown()) &&
+			(Chord.NeedsShift() ? InKeyEvent.IsShiftDown() : !InKeyEvent.IsShiftDown()) &&
+			(Chord.NeedsControl() ? InKeyEvent.IsControlDown() : !InKeyEvent.IsControlDown());
+
 		if (bIsMovingTimeSlider)
 		{
 			return true;
