Curve Editor: Panning

Added a new preference, Panning Mouse Button which currently is either Right or Alt-Middle mouse button. The default is right mouse button so there is no change in behavior.

New behaviors:
- Alt-right always zooms in/out

Fixed bugs:
- Fixed FCurveEditorDragOperation_PanView and FCurveEditorDragOperation_Zoom - LastMousePosition wasn't initialized in OnBeginDrag which caused constrain in up/down axis to fail
- Fixed FCurveEditorDragOperation_Zoom - Alt-right should constrain when shift is held down

Editing (insert keys, manipulating keys/tangents) in both cases is middle mouse button.
Removed middle mouse button binding to AddKeyHovered since that functionality is already handled. The user can rebind that to another hotkey.

diff --git a/Engine/Source/Editor/CurveEditor/Private/CurveEditorCommands.cpp b/Engine/Source/Editor/CurveEditor/Private/CurveEditorCommands.cpp
--- a/Engine/Source/Editor/CurveEditor/Private/CurveEditorCommands.cpp
+++ b/Engine/Source/Editor/CurveEditor/Private/CurveEditorCommands.cpp
@@ -73,7 +73,7 @@ void FCurveEditorCommands::RegisterCommands()
 	UI_COMMAND(ToggleShowBars, "Show Bars", "Show Bars like Constraints and Spaces", EUserInterfaceActionType::ToggleButton, FInputChord());
 	UI_COMMAND(ToggleShowValueIndicatorLines, "Show Value Indicators", "Show horizontal, dotted lines on the selected min & max valued key handles", EUserInterfaceActionType::ToggleButton, FInputChord());
 
-	UI_COMMAND(AddKeyHovered, "Add Key", "Add a new key to this curve at the current position.", EUserInterfaceActionType::Button, FInputChord(EKeys::MiddleMouseButton) );
+	UI_COMMAND(AddKeyHovered, "Add Key", "Add a new key to this curve at the current position.", EUserInterfaceActionType::Button, FInputChord() );
 	UI_COMMAND(PasteKeysHovered, "Paste", "Pastes the keys from the clipboard overwriting any key in destination track between the first and last pasted keys (default)", EUserInterfaceActionType::Button, FInputChord(EModifierKey::Control, EKeys::V));
 	UI_COMMAND(PasteAndMerge, "Paste & Merge", "Pastes the keys from the clipboard over the keys in the graph.", EUserInterfaceActionType::Button, FInputChord(EModifierKey::Control | EModifierKey::Shift, EKeys::V));
 	UI_COMMAND(PasteRelative, "Paste Relative", "Pastes the keys from the clipboard and aligns them to the nearest key to the left of the scrubber.", EUserInterfaceActionType::Button, FInputChord(EModifierKey::Alt, EKeys::V));
diff --git a/Engine/Source/Editor/CurveEditor/Private/CurveEditorSettings.cpp b/Engine/Source/Editor/CurveEditor/Private/CurveEditorSettings.cpp
--- a/Engine/Source/Editor/CurveEditor/Private/CurveEditorSettings.cpp
+++ b/Engine/Source/Editor/CurveEditor/Private/CurveEditorSettings.cpp
@@ -1,13 +1,14 @@
 // Copyright Epic Games, Inc. All Rights Reserved.
 
 #include "CurveEditorSettings.h"
-
 #include "HAL/PlatformCrt.h"
+#include "Input/Events.h"
 
 #include UE_INLINE_GENERATED_CPP_BY_NAME(CurveEditorSettings)
 
 UCurveEditorSettings::UCurveEditorSettings()
 {
+	PanningMouseButton = ECurveEditorPanningMouseButton::Right;
 	bAutoFrameCurveEditor = true;
 	bShowBars = true;
 	FrameInputPadding = 50;
@@ -26,6 +27,38 @@ UCurveEditorSettings::UCurveEditorSettings()
 	TreeViewWidth = 0.3f;
 }
 
+bool UCurveEditorSettings::AllowMousePan(const FPointerEvent& MouseEvent) const
+{
+	if (PanningMouseButton == ECurveEditorPanningMouseButton::AltMiddle)
+	{
+		return MouseEvent.IsMouseButtonDown(EKeys::MiddleMouseButton) && MouseEvent.IsAltDown();
+	}
+	else if (PanningMouseButton == ECurveEditorPanningMouseButton::Right)
+	{
+		return MouseEvent.IsMouseButtonDown(EKeys::RightMouseButton) && !MouseEvent.IsAltDown();
+	}
+
+	// Default to right mouse button
+	return MouseEvent.GetEffectingButton() == EKeys::RightMouseButton;
+}
+
+bool UCurveEditorSettings::AllowMouseEdit(const FPointerEvent& MouseEvent) const
+{
+	// If panning with alt-middle, middle allows mouse edits
+	if (PanningMouseButton == ECurveEditorPanningMouseButton::AltMiddle)
+	{
+		return MouseEvent.IsMouseButtonDown(EKeys::MiddleMouseButton) && !MouseEvent.IsAltDown();
+	}
+	// If panning with right, middle allows mouse edits
+	else if (PanningMouseButton == ECurveEditorPanningMouseButton::Right)
+	{
+		return MouseEvent.IsMouseButtonDown(EKeys::MiddleMouseButton);
+	}
+
+	// Default to middle mouse button
+	return MouseEvent.GetEffectingButton() == EKeys::MiddleMouseButton;
+}
+
 bool UCurveEditorSettings::GetAutoFrameCurveEditor() const
 {
 	return bAutoFrameCurveEditor;
diff --git a/Engine/Source/Editor/CurveEditor/Private/DragOperations/CurveEditorDragOperation_Pan.cpp b/Engine/Source/Editor/CurveEditor/Private/DragOperations/CurveEditorDragOperation_Pan.cpp
--- a/Engine/Source/Editor/CurveEditor/Private/DragOperations/CurveEditorDragOperation_Pan.cpp
+++ b/Engine/Source/Editor/CurveEditor/Private/DragOperations/CurveEditorDragOperation_Pan.cpp
@@ -29,6 +29,7 @@ void FCurveEditorDragOperation_PanView::OnBeginDrag(FVector2D InitialPosition, F
 	bIsDragging = true;
 
 	SnappingState.Reset();
+	LastMousePosition = CurrentPosition;
 }
 
 void FCurveEditorDragOperation_PanView::OnDrag(FVector2D InitialPosition, FVector2D CurrentPosition, const FPointerEvent& MouseEvent)
diff --git a/Engine/Source/Editor/CurveEditor/Private/DragOperations/CurveEditorDragOperation_Tangent.cpp b/Engine/Source/Editor/CurveEditor/Private/DragOperations/CurveEditorDragOperation_Tangent.cpp
--- a/Engine/Source/Editor/CurveEditor/Private/DragOperations/CurveEditorDragOperation_Tangent.cpp
+++ b/Engine/Source/Editor/CurveEditor/Private/DragOperations/CurveEditorDragOperation_Tangent.cpp
@@ -9,6 +9,7 @@
 #include "CurveEditorHelpers.h"
 #include "CurveEditorScreenSpace.h"
 #include "CurveEditorSelection.h"
+#include "CurveEditorSettings.h"
 #include "CurveModel.h"
 #include "Curves/RichCurve.h"
 #include "HAL/PlatformCrt.h"
@@ -33,6 +34,7 @@ namespace CurveEditorDragOperation
 	{}
 
 	void FTangentSolver::Process(
+		FCurveEditor* CurveEditor,
 		const FPointerEvent& MouseEvent, 
 		const FCurveEditorScreenSpace& CurveSpace, 
 		const FVector2D& PixelDelta)
@@ -45,6 +47,7 @@ namespace CurveEditorDragOperation
 			const float ArriveTangent = InitialKeyAttributes.GetArriveTangent();
 
 			const FSolverParams SolverParams(
+				CurveEditor,
 				MouseEvent,
 				CurveSpace,
 				PixelDelta,
@@ -69,6 +72,7 @@ namespace CurveEditorDragOperation
 			const float LeaveTangent = InitialKeyAttributes.GetLeaveTangent();
 
 			const FSolverParams SolverParams(
+				CurveEditor,
 				MouseEvent,
 				CurveSpace,
 				PixelDelta,
@@ -89,13 +93,15 @@ namespace CurveEditorDragOperation
 	}
 
 	FTangentSolver::FSolverParams::FSolverParams(
+		FCurveEditor* InCurveEditor,
 		const FPointerEvent& InMouseEvent,
 		const FCurveEditorScreenSpace& InCurveSpace,
 		const FVector2D& InPixelDelta,
 		const float InDisplayRatio,
 		const bool bInArriveTangent,
 		const float InInitialTangent)
-		: MouseEvent(InMouseEvent)
+		: CurveEditor(InCurveEditor)
+		, MouseEvent(InMouseEvent)
 		, CurveSpace(InCurveSpace)
 		, PixelDelta(InPixelDelta)
 		, DisplayRatio(InDisplayRatio)
@@ -106,7 +112,7 @@ namespace CurveEditorDragOperation
 	void FTangentSolver::SolveTangentConstantWeight(const FSolverParams& Params)
 	{
 		if (Params.MouseEvent.IsMouseButtonDown(EKeys::LeftMouseButton) ||
-			Params.MouseEvent.IsMouseButtonDown(EKeys::MiddleMouseButton))
+			Params.CurveEditor->GetSettings()->AllowMouseEdit(Params.MouseEvent))
 		{				
 			const float PixelLength = Params.bArriveTangent ? -60.0f : 60.f;
 			const float Slope = Params.InitialTangent * -Params.DisplayRatio;
@@ -143,7 +149,7 @@ namespace CurveEditorDragOperation
 		FVector2D NewTangentOffset;
 
 		if (Params.MouseEvent.IsMouseButtonDown(EKeys::LeftMouseButton) ||
-			Params.MouseEvent.IsMouseButtonDown(EKeys::MiddleMouseButton))
+			Params.CurveEditor->GetSettings()->AllowMouseEdit(Params.MouseEvent))
 		{
 			const FVector2D PreviousTangentOffset = LastTangentOffset.IsSet() ? LastTangentOffset.GetValue() : InitialTangentOffset;
 			const FVector2D PointAtOffset = InitialTangentOffset + Params.PixelDelta;
@@ -288,7 +294,7 @@ void FCurveEditorDragOperation_Tangent::OnDrag(FVector2D InitialPosition, FVecto
 
 		for (CurveEditorDragOperation::FTangentSolver& Solver : KeyData.TangentSolvers)
 		{
-			Solver.Process(MouseEvent, CurveSpace, PixelDelta);
+			Solver.Process(CurveEditor, MouseEvent, CurveSpace, PixelDelta);
 
 			NewKeyAttributesScratch.Add(Solver.GetKeyAttributes());
 		}
diff --git a/Engine/Source/Editor/CurveEditor/Private/DragOperations/CurveEditorDragOperation_Tangent.h b/Engine/Source/Editor/CurveEditor/Private/DragOperations/CurveEditorDragOperation_Tangent.h
--- a/Engine/Source/Editor/CurveEditor/Private/DragOperations/CurveEditorDragOperation_Tangent.h
+++ b/Engine/Source/Editor/CurveEditor/Private/DragOperations/CurveEditorDragOperation_Tangent.h
@@ -30,7 +30,7 @@ namespace CurveEditorDragOperation
 		 * @param CurveSpace			The screen space of the curve editor
 		 * @param PixelDelta			The mouse position delta from the initial mouse position in pixels
 		 */
-		void Process(const FPointerEvent& MouseEvent, const FCurveEditorScreenSpace& CurveSpace, const FVector2D& PixelDelta);
+		void Process(FCurveEditor* CurveEditor, const FPointerEvent& MouseEvent, const FCurveEditorScreenSpace& CurveSpace, const FVector2D& PixelDelta);
 
 		/** Returns the current key attributes */
 		const FKeyAttributes& GetKeyAttributes() const { return KeyAttributes; }
@@ -43,6 +43,7 @@ namespace CurveEditorDragOperation
 		struct FSolverParams
 		{
 			FSolverParams(
+				FCurveEditor* InCurveEditor,
 				const FPointerEvent& InMouseEvent,
 				const FCurveEditorScreenSpace& InCurveSpace,
 				const FVector2D& InPixelDelta,
@@ -50,6 +51,9 @@ namespace CurveEditorDragOperation
 				const bool bInArriveTangent,
 				const float InInitialTangent);
 
+			/** The curve editor */
+			FCurveEditor* CurveEditor;
+
 			/** The current mouse event */
 			const FPointerEvent& MouseEvent;
 
diff --git a/Engine/Source/Editor/CurveEditor/Private/DragOperations/CurveEditorDragOperation_Zoom.cpp b/Engine/Source/Editor/CurveEditor/Private/DragOperations/CurveEditorDragOperation_Zoom.cpp
--- a/Engine/Source/Editor/CurveEditor/Private/DragOperations/CurveEditorDragOperation_Zoom.cpp
+++ b/Engine/Source/Editor/CurveEditor/Private/DragOperations/CurveEditorDragOperation_Zoom.cpp
@@ -43,12 +43,18 @@ void FCurveEditorDragOperation_Zoom::OnBeginDrag(FVector2D InitialPosition, FVec
 		OriginalOutputRange = (OutputSpace.GetOutputMax() - OutputSpace.GetOutputMin());
 		ZoomOriginY = OutputSpace.GetOutputMin() + OriginalOutputRange * (1.f - ZoomFactor.Y);
 	}
+
+	SnappingState.Reset();
+	LastMousePosition = CurrentPosition;
 }
 
 void FCurveEditorDragOperation_Zoom::OnDrag(FVector2D InitialPosition, FVector2D CurrentPosition, const FPointerEvent& MouseEvent)
 {
 	constexpr double ClampRange = 1e9f;
-	const FVector2D PixelDelta = CurrentPosition - InitialPosition;
+
+	const FVector2D PixelDelta = CurveEditor->GetAxisSnap().GetSnappedPosition(InitialPosition, LastMousePosition, CurrentPosition, MouseEvent, SnappingState, true) - InitialPosition;
+	LastMousePosition = CurrentPosition;
+
 	const FCurveEditorZoomScaleConfig& ZoomConfig = CurveEditor->GetZoomScaleConfig();
 
 	// Zoom input range
diff --git a/Engine/Source/Editor/CurveEditor/Private/DragOperations/CurveEditorDragOperation_Zoom.h b/Engine/Source/Editor/CurveEditor/Private/DragOperations/CurveEditorDragOperation_Zoom.h
--- a/Engine/Source/Editor/CurveEditor/Private/DragOperations/CurveEditorDragOperation_Zoom.h
+++ b/Engine/Source/Editor/CurveEditor/Private/DragOperations/CurveEditorDragOperation_Zoom.h
@@ -29,6 +29,9 @@ private:
 	double ZoomOriginX, ZoomOriginY;
 	double OriginalInputRange, OriginalOutputRange;
 
+	FCurveEditorAxisSnap::FSnapState SnappingState;
+	FVector2D LastMousePosition;
+
 	double GetZoomMultiplier_InputAxis(double InMovedMouseX) const;
 	double GetZoomMultiplier_OutputAxis(double InMovedMouseY) const;
 };
diff --git a/Engine/Source/Editor/CurveEditor/Private/SCurveEditorViewContainer.cpp b/Engine/Source/Editor/CurveEditor/Private/SCurveEditorViewContainer.cpp
--- a/Engine/Source/Editor/CurveEditor/Private/SCurveEditorViewContainer.cpp
+++ b/Engine/Source/Editor/CurveEditor/Private/SCurveEditorViewContainer.cpp
@@ -272,7 +272,7 @@ FReply SCurveEditorViewContainer::OnMouseButtonDown(const FGeometry& MyGeometry,
 			return FReply::Handled();
 		}
 	}
-	else if (MouseEvent.GetEffectingButton() == EKeys::RightMouseButton)
+	else if (CurveEditor->GetSettings()->AllowMousePan(MouseEvent))
 	{
 		// Zoom Timeline
 		if (MouseEvent.IsAltDown())
diff --git a/Engine/Source/Editor/CurveEditor/Private/Views/SInteractiveCurveEditorView.cpp b/Engine/Source/Editor/CurveEditor/Private/Views/SInteractiveCurveEditorView.cpp
--- a/Engine/Source/Editor/CurveEditor/Private/Views/SInteractiveCurveEditorView.cpp
+++ b/Engine/Source/Editor/CurveEditor/Private/Views/SInteractiveCurveEditorView.cpp
@@ -1600,138 +1600,113 @@ FReply SInteractiveCurveEditorView::OnMouseButtonDown(const FGeometry& MyGeometr
 	// Rebind our context actions so that shift click commands use the right position.
 	RebindContextualActions(MousePixel);
 
-	// // Marquee Selection
-	// if (MouseEvent.GetEffectingButton() == EKeys::LeftMouseButton)
-	// {
-	// 	DragOperation = FCurveEditorDelayedDrag(MousePixel, MouseEvent.GetEffectingButton());
-	// 	DragOperation->DragImpl = MakeUnique<FCurveEditorDragOperation_Marquee>(CurveEditor.Get(), this);
-	// 	return FReply::Handled();
-	// }
-	// Middle Click + Alt Pan
-	//else
-	if (MouseEvent.GetEffectingButton() == EKeys::MiddleMouseButton)
+	if (CurveEditor->GetSettings()->AllowMouseEdit(MouseEvent))
 	{
-		if (MouseEvent.IsAltDown())
-		{
-			DragOperation = FCurveEditorDelayedDrag(MousePixel, MouseEvent.GetEffectingButton());
-			DragOperation->DragImpl = MakeUnique<FCurveEditorDragOperation_PanInput>(CurveEditor.Get());
-			return FReply::Handled();
-		}
-		else
-		{
-			// Middle Mouse can try to create keys on curves.
-			TOptional<FCurvePointHandle> NewPoint;
+		// Middle Mouse can try to create keys on curves.
+		TOptional<FCurvePointHandle> NewPoint;
 
-			// Add a key to the closest curve to the mouse
-			if (TOptional<FCurveModelID> HoveredCurve = GetHoveredCurve())
-			{
-				// Don't allow adding keys when shift is held down with selected keys since that is for dragging keys in a constrained axis
-				const bool bDraggingKeys = MouseEvent.IsShiftDown() && !CurveEditor->GetSelection().GetAll().IsEmpty();
+		// Add a key to the closest curve to the mouse
+		if (TOptional<FCurveModelID> HoveredCurve = GetHoveredCurve())
+		{
+			// Don't allow adding keys when shift is held down with selected keys since that is for dragging keys in a constrained axis
+			const bool bDraggingKeys = MouseEvent.IsShiftDown() && !CurveEditor->GetSelection().GetAll().IsEmpty();
 
-				// Don't create a new key if there is already a key or transform handle in place
-				const bool bKeyAlreadyExists = HitPoint(MousePixel).IsSet();
+			// Don't create a new key if there is already a key or transform handle in place
+			const bool bKeyAlreadyExists = HitPoint(MousePixel).IsSet();
 
-				FCurveModel* CurveToAddTo = CurveEditor->FindCurve(HoveredCurve.GetValue());
-				if (CurveToAddTo && !CurveToAddTo->IsReadOnly() && !bDraggingKeys && !bKeyAlreadyExists)
-				{
-					FScopedTransaction Transaction(LOCTEXT("InsertKey", "Insert Key")); // Selection and key change should be part of same undo / redo op.
-					UE::CurveEditor::FScopedSelectionChange SelectionChange(CurveEditor);
+			FCurveModel* CurveToAddTo = CurveEditor->FindCurve(HoveredCurve.GetValue());
+			if (CurveToAddTo && !CurveToAddTo->IsReadOnly() && !bDraggingKeys && !bKeyAlreadyExists)
+			{
+				FScopedTransaction Transaction(LOCTEXT("InsertKey", "Insert Key")); // Selection and key change should be part of same undo / redo op.
+				UE::CurveEditor::FScopedSelectionChange SelectionChange(CurveEditor);
 
-					FCurveEditorScreenSpace CurveSpace = GetCurveSpace(HoveredCurve.GetValue());
-					double MouseTime = CurveSpace.ScreenToSeconds(MousePixel.X);
-					double MouseValue = CurveSpace.ScreenToValue(MousePixel.Y);
+				FCurveEditorScreenSpace CurveSpace = GetCurveSpace(HoveredCurve.GetValue());
+				double MouseTime = CurveSpace.ScreenToSeconds(MousePixel.X);
+				double MouseValue = CurveSpace.ScreenToValue(MousePixel.Y);
 
-					FKeyAttributes KeyAttributes = GetDefaultKeyAttributesForCurveTime(*CurveEditor, *CurveToAddTo, MouseTime);
+				FKeyAttributes KeyAttributes = GetDefaultKeyAttributesForCurveTime(*CurveEditor, *CurveToAddTo, MouseTime);
 
-					FCurveSnapMetrics SnapMetrics = CurveEditor->GetCurveSnapMetrics(HoveredCurve.GetValue());
-					MouseTime = SnapMetrics.SnapInputSeconds(MouseTime);
-					MouseValue = SnapMetrics.SnapOutput(MouseValue);
+				FCurveSnapMetrics SnapMetrics = CurveEditor->GetCurveSnapMetrics(HoveredCurve.GetValue());
+				MouseTime = SnapMetrics.SnapInputSeconds(MouseTime);
+				MouseValue = SnapMetrics.SnapOutput(MouseValue);
 
-					// If control is pressed. Keep the curve unchanged
-					if (MouseEvent.IsControlDown())
-					{
-						KeyAttributes.SetTangentMode(RCTM_User);
+				// If control is pressed. Keep the curve unchanged
+				if (MouseEvent.IsControlDown())
+				{
+					KeyAttributes.SetTangentMode(RCTM_User);
 
-						// Estimated delta time to compute right and left tangents
-						double DeltaTime = 0.1;
+					// Estimated delta time to compute right and left tangents
+					double DeltaTime = 0.1;
 
-						// Make mouse value more accurate 
-						CurveToAddTo->Evaluate(MouseTime, MouseValue);
+					// Make mouse value more accurate 
+					CurveToAddTo->Evaluate(MouseTime, MouseValue);
 
-						// Compute right tangent
-						double RightTangent = GetTangentValue(MouseTime, MouseValue, CurveToAddTo, DeltaTime);
-						KeyAttributes.SetLeaveTangent(RightTangent);
+					// Compute right tangent
+					double RightTangent = GetTangentValue(MouseTime, MouseValue, CurveToAddTo, DeltaTime);
+					KeyAttributes.SetLeaveTangent(RightTangent);
 
-						// Left
-						double LeftTangent = GetTangentValue(MouseTime, MouseValue, CurveToAddTo, -DeltaTime);
-						KeyAttributes.SetArriveTangent(LeftTangent);
-					}
+					// Left
+					double LeftTangent = GetTangentValue(MouseTime, MouseValue, CurveToAddTo, -DeltaTime);
+					KeyAttributes.SetArriveTangent(LeftTangent);
+				}
 
-					// When adding to a curve with no variance, add it with the same value so that
-					// curves don't pop wildly in normalized views due to a slight difference between the keys
-					double CurveOutputMin = 0, CurveOutputMax = 1;
-					CurveToAddTo->GetValueRange(CurveOutputMin, CurveOutputMax);
-					if (CurveOutputMin == CurveOutputMax)
-					{
-						MouseValue = CurveOutputMin;
-					}
+				// When adding to a curve with no variance, add it with the same value so that
+				// curves don't pop wildly in normalized views due to a slight difference between the keys
+				double CurveOutputMin = 0, CurveOutputMax = 1;
+				CurveToAddTo->GetValueRange(CurveOutputMin, CurveOutputMax);
+				if (CurveOutputMin == CurveOutputMax)
+				{
+					MouseValue = CurveOutputMin;
+				}
 
-					CurveToAddTo->Modify();
+				CurveToAddTo->Modify();
 
-					// Add a key on this curve
-					TOptional<FKeyHandle> NewKey = CurveToAddTo->AddKey(FKeyPosition(MouseTime, MouseValue), KeyAttributes);
-					if (NewKey.IsSet())
-					{
-						NewPoint = FCurvePointHandle(HoveredCurve.GetValue(), ECurvePointType::Key, NewKey.GetValue());
+				// Add a key on this curve
+				TOptional<FKeyHandle> NewKey = CurveToAddTo->AddKey(FKeyPosition(MouseTime, MouseValue), KeyAttributes);
+				if (NewKey.IsSet())
+				{
+					NewPoint = FCurvePointHandle(HoveredCurve.GetValue(), ECurvePointType::Key, NewKey.GetValue());
 
-						CurveEditor->GetSelection().Clear();
-						CurveEditor->GetSelection().Add(NewPoint.GetValue());
-					}
-					else
-					{
-						Transaction.Cancel();
-					}
+					CurveEditor->GetSelection().Clear();
+					CurveEditor->GetSelection().Add(NewPoint.GetValue());
+				}
+				else
+				{
+					Transaction.Cancel();
 				}
 			}
+		}
 
-			// If there are any tangent handles selected, prefer to drag those instead of keys
-			ECurvePointType PointType = ECurvePointType::Key;
-			if (!NewPoint.IsSet())
+		// If there are any tangent handles selected, prefer to drag those instead of keys
+		ECurvePointType PointType = ECurvePointType::Key;
+		if (!NewPoint.IsSet())
+		{
+			for (const TTuple<FCurveModelID, FKeyHandleSet>& Pair : CurveEditor->GetSelection().GetAll())
 			{
-				for (const TTuple<FCurveModelID, FKeyHandleSet>& Pair : CurveEditor->GetSelection().GetAll())
+				for (FKeyHandle Handle : Pair.Value.AsArray())
 				{
-					for (FKeyHandle Handle : Pair.Value.AsArray())
+					if (Pair.Value.Contains(Handle, ECurvePointType::ArriveTangent) || Pair.Value.Contains(Handle, ECurvePointType::LeaveTangent))
 					{
-						if (Pair.Value.Contains(Handle, ECurvePointType::ArriveTangent) || Pair.Value.Contains(Handle, ECurvePointType::LeaveTangent))
-						{
-							PointType = ECurvePointType::ArriveTangent;
-							break;
-						}
+						PointType = ECurvePointType::ArriveTangent;
+						break;
 					}
 				}
 			}
+		}
 
-			TUniquePtr<ICurveEditorKeyDragOperation> KeyDrag = CreateKeyDrag(PointType);
+		TUniquePtr<ICurveEditorKeyDragOperation> KeyDrag = CreateKeyDrag(PointType);
 
-			KeyDrag->Initialize(CurveEditor.Get(), NewPoint);
+		KeyDrag->Initialize(CurveEditor.Get(), NewPoint);
 
-			DragOperation = FCurveEditorDelayedDrag(MousePixel, MouseEvent.GetEffectingButton());
-			DragOperation->DragImpl = MoveTemp(KeyDrag);
+		DragOperation = FCurveEditorDelayedDrag(MousePixel, MouseEvent.GetEffectingButton());
+		DragOperation->DragImpl = MoveTemp(KeyDrag);
 
-			return FReply::Handled().PreventThrottling();
-		}
+		return FReply::Handled().PreventThrottling();
 	}
-	else if (MouseEvent.GetEffectingButton() == EKeys::RightMouseButton)
+	else if (CurveEditor->GetSettings()->AllowMousePan(MouseEvent))
 	{
-		// Zoom Timeline
-		if (MouseEvent.IsAltDown())
-		{
-			DragOperation = FCurveEditorDelayedDrag(MousePixel, MouseEvent.GetEffectingButton());
-			DragOperation->DragImpl = MakeUnique<FCurveEditorDragOperation_Zoom>(CurveEditor.Get(), SharedThis(this));
-			return FReply::Handled();
-		}
 		// Pan Timeline if we have flexible output bounds
-		else if (!bFixedOutputBounds)
+		if (!bFixedOutputBounds)
 		{
 			DragOperation = FCurveEditorDelayedDrag(MousePixel, MouseEvent.GetEffectingButton());
 			DragOperation->DragImpl = MakeUnique<FCurveEditorDragOperation_PanView>(CurveEditor.Get(), SharedThis(this));
@@ -1763,6 +1738,16 @@ FReply SInteractiveCurveEditorView::OnMouseButtonDown(const FGeometry& MyGeometr
 			}
 		}
 	}
+	else if (MouseEvent.GetEffectingButton() == EKeys::RightMouseButton)
+	{
+		// Zoom Timeline
+		if (MouseEvent.IsAltDown())
+		{
+			DragOperation = FCurveEditorDelayedDrag(MousePixel, MouseEvent.GetEffectingButton());
+			DragOperation->DragImpl = MakeUnique<FCurveEditorDragOperation_Zoom>(CurveEditor.Get(), SharedThis(this));
+			return FReply::Handled();
+		}
+	}
 
 	return FReply::Unhandled();
 }
@@ -1777,23 +1762,21 @@ FReply SInteractiveCurveEditorView::OnMouseButtonUp(const FGeometry& MyGeometry,
 	}
 
 	const bool bDragOperationWasSet = DragOperation.IsSet();
+	const bool bWasDragging = DragOperation.IsSet() && DragOperation->IsDragging();
 	FVector2D MousePosition = MyGeometry.AbsoluteToLocal(MouseEvent.GetScreenSpacePosition());
-	if (DragOperation.IsSet())
+	if (bWasDragging)
 	{
-		if (DragOperation->IsDragging())
-		{
-			FVector2D InitialPosition = DragOperation->GetInitialPosition();
-			DragOperation->DragImpl->EndDrag(InitialPosition, MousePosition, MouseEvent);
+		FVector2D InitialPosition = DragOperation->GetInitialPosition();
+		DragOperation->DragImpl->EndDrag(InitialPosition, MousePosition, MouseEvent);
 
-			DragOperation.Reset();
-			return FReply::Handled().ReleaseMouseCapture();
-		}
+		DragOperation.Reset();
+		return FReply::Handled().ReleaseMouseCapture();
 	}
 
 	DragOperation.Reset();
 
 	// Select the curve on mouse release if no key or tangent was clicked on
-	if (!bDragOperationWasSet && MouseEvent.GetEffectingButton() == EKeys::LeftMouseButton)
+	if (!bWasDragging && MouseEvent.GetEffectingButton() == EKeys::LeftMouseButton)
 	{
 		const bool bIsShiftDown = MouseEvent.IsShiftDown();
 		const bool bIsAltDown = MouseEvent.IsAltDown();
@@ -1833,7 +1816,8 @@ FReply SInteractiveCurveEditorView::OnMouseButtonUp(const FGeometry& MyGeometry,
 		}
 	}
 
-	if (MouseEvent.GetEffectingButton() == EKeys::RightMouseButton)
+	// Pop up a menu if not dragging
+	if (!bWasDragging && MouseEvent.GetEffectingButton() == EKeys::RightMouseButton)
 	{
 		FVector2D MousePixel = MyGeometry.AbsoluteToLocal(MouseEvent.GetScreenSpacePosition());
 		TOptional<FCurvePointHandle> MouseDownPoint = HitPoint(MousePixel);
diff --git a/Engine/Source/Editor/CurveEditor/Public/CurveEditorSettings.h b/Engine/Source/Editor/CurveEditor/Public/CurveEditorSettings.h
--- a/Engine/Source/Editor/CurveEditor/Public/CurveEditorSettings.h
+++ b/Engine/Source/Editor/CurveEditor/Public/CurveEditorSettings.h
@@ -15,10 +15,19 @@
 
 #include "CurveEditorSettings.generated.h"
 
+struct FPointerEvent;
+
 #define UE_API CURVEEDITOR_API
 
 class UClass;
 
+UENUM()
+enum class ECurveEditorPanningMouseButton : uint8
+{
+	Right	  UMETA(DisplayName = "Right"),
+	AltMiddle UMETA(DisplayName = "Alt-Middle")
+};
+
 /** Defines visibility states for the tangents in the curve editor. */
 UENUM()
 enum class ECurveEditorTangentVisibility : uint8
@@ -105,6 +114,10 @@ public:
 
 	UE_API UCurveEditorSettings();
 
+	/** Get whether the mouse event allows panning or editing */
+	bool AllowMousePan(const FPointerEvent& MouseEvent) const;
+	bool AllowMouseEdit(const FPointerEvent& MouseEvent) const;
+
 	/** Gets whether or not the curve editor auto frames the selected curves. */
 	UE_API bool GetAutoFrameCurveEditor() const;
 	/** Sets whether or not the curve editor auto frames the selected curves. */
@@ -203,6 +216,9 @@ protected:
 
 	UE_API virtual void PostEditChangeProperty(FPropertyChangedEvent& PropertyChangedEvent) override;
 
+	UPROPERTY(config, EditAnywhere, Category = "Curve Editor")
+	ECurveEditorPanningMouseButton PanningMouseButton;
+
 	UPROPERTY( config, EditAnywhere, Category="Curve Editor" )
 	bool bAutoFrameCurveEditor;
 
