Curve Editor: CreateKeyProxies now takes an FCurveEditor and a Curve Model ID

This fixes a crash in RichCurveKeyProxy holding a raw pointer to RichCurveEditorModel. This came in through 10950907

diff --git a/Engine/Plugins/Animation/ControlRig/Source/ControlRigEditor/Private/ControlRigSpaceChannelCurveModel.cpp b/Engine/Plugins/Animation/ControlRig/Source/ControlRigEditor/Private/ControlRigSpaceChannelCurveModel.cpp
--- a/Engine/Plugins/Animation/ControlRig/Source/ControlRigEditor/Private/ControlRigSpaceChannelCurveModel.cpp
+++ b/Engine/Plugins/Animation/ControlRig/Source/ControlRigEditor/Private/ControlRigSpaceChannelCurveModel.cpp
@@ -267,7 +267,7 @@ void FControlRigSpaceChannelCurveModel::SetCurveAttributes(const FCurveAttribute
 {
 }
 
-void FControlRigSpaceChannelCurveModel::CreateKeyProxies(TArrayView<const FKeyHandle> InKeyHandles, TArrayView<UObject*> OutObjects)
+void FControlRigSpaceChannelCurveModel::CreateKeyProxies(TWeakPtr<FCurveEditor> InWeakCurveEditor, FCurveModelID InCurveModelID, TArrayView<const FKeyHandle> InKeyHandles, TArrayView<UObject*> OutObjects)
 {
 }
 
diff --git a/Engine/Plugins/Animation/ControlRig/Source/ControlRigEditor/Public/ControlRigSpaceChannelCurveModel.h b/Engine/Plugins/Animation/ControlRig/Source/ControlRigEditor/Public/ControlRigSpaceChannelCurveModel.h
--- a/Engine/Plugins/Animation/ControlRig/Source/ControlRigEditor/Public/ControlRigSpaceChannelCurveModel.h
+++ b/Engine/Plugins/Animation/ControlRig/Source/ControlRigEditor/Public/ControlRigSpaceChannelCurveModel.h
@@ -45,7 +45,7 @@ public:
 	virtual void AddKeys(TArrayView<const FKeyPosition> InKeyPositions, TArrayView<const FKeyAttributes> InAttributes, TArrayView<TOptional<FKeyHandle>>* OutKeyHandles) override;
 	virtual void RemoveKeys(TArrayView<const FKeyHandle> InKeys, double InCurrentTime) override;
 
-	virtual void CreateKeyProxies(TArrayView<const FKeyHandle> InKeyHandles, TArrayView<UObject*> OutObjects) override;
+	virtual void CreateKeyProxies(TWeakPtr<FCurveEditor> InWeakCurveEditor, FCurveModelID InCurveModelID, TArrayView<const FKeyHandle> InKeyHandles, TArrayView<UObject*> OutObjects) override;
 	
 	virtual UObject* GetOwningObject() const override;
 	virtual bool HasChangedAndResetTest() override;
diff --git a/Engine/Source/Editor/CurveEditor/Private/RichCurveEditorModel.cpp b/Engine/Source/Editor/CurveEditor/Private/RichCurveEditorModel.cpp
--- a/Engine/Source/Editor/CurveEditor/Private/RichCurveEditorModel.cpp
+++ b/Engine/Source/Editor/CurveEditor/Private/RichCurveEditorModel.cpp
@@ -669,22 +669,23 @@ void FRichCurveEditorModel::SetCurveAttributes(const FCurveAttributes& InCurveAt
 	}
 }
 
-void FRichCurveEditorModel::CreateKeyProxies(TArrayView<const FKeyHandle> InKeyHandles, TArrayView<UObject*> OutObjects)
-{
-	if (UObject* Owner = WeakOwner.Get())
-	{
-		if(IsValid())
-		{
-			for (int32 Index = 0; Index < InKeyHandles.Num(); ++Index)
-			{
-				URichCurveKeyProxy* NewProxy = NewObject<URichCurveKeyProxy>(GetTransientPackage(), NAME_None);
+ void FRichCurveEditorModel::CreateKeyProxies(TWeakPtr<FCurveEditor> InWeakCurveEditor, FCurveModelID InCurveModelID, TArrayView<const FKeyHandle> InKeyHandles, TArrayView<UObject*> OutObjects)
+ {
+	 if (UObject* Owner = WeakOwner.Get())
+	 {
+		 if (IsValid())
+		 {
+			 for (int32 Index = 0; Index < InKeyHandles.Num(); ++Index)
+			 {
+				 URichCurveKeyProxy* NewProxy = NewObject<URichCurveKeyProxy>(GetTransientPackage(), NAME_None);
+
+				 NewProxy->Initialize(InWeakCurveEditor, InCurveModelID, InKeyHandles[Index], WeakOwner);
+				 OutObjects[Index] = NewProxy;
+			 }
+		 }
+	 }
+ }
 
-				NewProxy->Initialize(InKeyHandles[Index], this, WeakOwner);
-				OutObjects[Index] = NewProxy;
-			}
-		}
-	}
-}
 
 TUniquePtr<IBufferedCurveModel> FRichCurveEditorModel::CreateBufferedCurveCopy() const
 {
diff --git a/Engine/Source/Editor/CurveEditor/Private/RichCurveKeyProxy.h b/Engine/Source/Editor/CurveEditor/Private/RichCurveKeyProxy.h
--- a/Engine/Source/Editor/CurveEditor/Private/RichCurveKeyProxy.h
+++ b/Engine/Source/Editor/CurveEditor/Private/RichCurveKeyProxy.h
@@ -5,6 +5,7 @@
 #include "CoreTypes.h"
 #include "UObject/WeakObjectPtr.h"
 #include "UObject/WeakObjectPtrTemplates.h"
+#include "CurveEditor.h"
 #include "CurveEditorKeyProxy.h"
 #include "Curves/RichCurve.h"
 #include "RichCurveEditorModel.h"
@@ -20,20 +21,25 @@ public:
 	/**
 	 * Initialize this key proxy object by caching the underlying key object, and retrieving the time/value each tick
 	 */
-	void Initialize(FKeyHandle InKeyHandle, FRichCurveEditorModel* InModel, TWeakObjectPtr<UObject> InWeakOwner)
+	void Initialize(TWeakPtr<FCurveEditor> InWeakCurveEditor, FCurveModelID InCurveModelID, FKeyHandle InKeyHandle, TWeakObjectPtr<UObject> InWeakOwner)
 	{
+		WeakCurveEditor = InWeakCurveEditor;
+		CurveModelID = InCurveModelID;
 		KeyHandle  = InKeyHandle;
 		WeakOwner  = InWeakOwner;
-		Model      = InModel;
 
 		if(UObject* Owner = WeakOwner.Get())
 		{
-			if(Model->IsValid())
+			TSharedPtr<FCurveEditor> CurveEditor = WeakCurveEditor.Pin();
+			if (CurveEditor)
 			{
-				const FRichCurve& RichCurve = Model->GetReadOnlyRichCurve();
-				if (RichCurve.IsKeyHandleValid(KeyHandle))
+				if (FRichCurveEditorModel* Model = static_cast<FRichCurveEditorModel*>(CurveEditor->FindCurve(CurveModelID)))
 				{
-					Value = RichCurve.GetKey(KeyHandle);
+					const FRichCurve& RichCurve = Model->GetReadOnlyRichCurve();
+					if (RichCurve.IsKeyHandleValid(KeyHandle))
+					{
+						Value = RichCurve.GetKey(KeyHandle);
+					}
 				}
 			}
 		}
@@ -47,19 +53,23 @@ private:
 
 		if(UObject* Owner = WeakOwner.Get())
 		{
-			if(Model->IsValid())
+			TSharedPtr<FCurveEditor> CurveEditor = WeakCurveEditor.Pin();
+			if (CurveEditor)
 			{
-				FRichCurve& RichCurve = Model->GetRichCurve();
-				if(RichCurve.IsKeyHandleValid(KeyHandle))
+				if (FRichCurveEditorModel* Model = static_cast<FRichCurveEditorModel*>(CurveEditor->FindCurve(CurveModelID)))
 				{
-					Owner->Modify();
+					FRichCurve& RichCurve = Model->GetRichCurve();
+					if(RichCurve.IsKeyHandleValid(KeyHandle))
+					{
+						Owner->Modify();
 
-					FRichCurveKey& ActualKey = RichCurve.GetKey(KeyHandle);
-					const FKeyPosition KeyPosition(Value.Time, Value.Value);
-					Model->SetKeyPositions({KeyHandle}, {KeyPosition}, PropertyChangedEvent.ChangeType);
+						FRichCurveKey& ActualKey = RichCurve.GetKey(KeyHandle);
+						const FKeyPosition KeyPosition(Value.Time, Value.Value);
+						Model->SetKeyPositions({KeyHandle}, {KeyPosition}, PropertyChangedEvent.ChangeType);
 				
-					Owner->PostEditChangeProperty(PropertyChangedEvent);
-					Model->OnCurveModified().Broadcast();
+						Owner->PostEditChangeProperty(PropertyChangedEvent);
+						Model->OnCurveModified().Broadcast();
+					}
 				}
 			}
 		}	
@@ -69,12 +79,16 @@ private:
 	{
 		if(UObject* Owner = WeakOwner.Get())
 		{
-			if(Model->IsValid())
+			TSharedPtr<FCurveEditor> CurveEditor = WeakCurveEditor.Pin();
+			if (CurveEditor)
 			{
-				const FRichCurve& RichCurve = Model->GetReadOnlyRichCurve();
-				if (RichCurve.IsKeyHandleValid(KeyHandle))
+				if (FRichCurveEditorModel* Model = static_cast<FRichCurveEditorModel*>(CurveEditor->FindCurve(CurveModelID)))
 				{
-					Value = RichCurve.GetKey(KeyHandle);
+					const FRichCurve& RichCurve = Model->GetReadOnlyRichCurve();
+					if (RichCurve.IsKeyHandleValid(KeyHandle))
+					{
+						Value = RichCurve.GetKey(KeyHandle);
+					}
 				}
 			}
 		}
@@ -88,10 +102,12 @@ private:
 
 private:
 
+	/** Cached curve editor */
+	TWeakPtr<FCurveEditor> WeakCurveEditor;
+	/** Cached curve model ID for this rich curve's key */
+	FCurveModelID CurveModelID;
 	/** Cached key handle that this key proxy relates to */
 	FKeyHandle KeyHandle;
-	/** Cached model that created this proxy */
-	FRichCurveEditorModel* Model;
 	/** Cached owner in which the raw curve resides */
 	TWeakObjectPtr<UObject> WeakOwner;
 };
\ No newline at end of file
diff --git a/Engine/Source/Editor/CurveEditor/Private/SCurveEditorPanel.cpp b/Engine/Source/Editor/CurveEditor/Private/SCurveEditorPanel.cpp
--- a/Engine/Source/Editor/CurveEditor/Private/SCurveEditorPanel.cpp
+++ b/Engine/Source/Editor/CurveEditor/Private/SCurveEditorPanel.cpp
@@ -887,7 +887,7 @@ void SCurveEditorPanel::UpdateEditBox()
 		if (KeyHandleScratch.Num() > 0)
 		{
 			NewProxiesScratch.SetNum(KeyHandleScratch.Num());
-			Curve->CreateKeyProxies(KeyHandleScratch, NewProxiesScratch);
+			Curve->CreateKeyProxies(CurveEditor, Pair.Key, KeyHandleScratch, NewProxiesScratch);
 
 			for (int32 Index = 0; Index < KeyHandleScratch.Num(); ++Index)
 			{
diff --git a/Engine/Source/Editor/CurveEditor/Public/CurveModel.h b/Engine/Source/Editor/CurveEditor/Public/CurveModel.h
--- a/Engine/Source/Editor/CurveEditor/Public/CurveModel.h
+++ b/Engine/Source/Editor/CurveEditor/Public/CurveModel.h
@@ -347,9 +347,15 @@ public:
 	/**
 	 * Create key proxy objects for the specified key handles. One object should be assigned to OutObjects per index within InKeyHandles
 	 *
+	 * @param InWeakCurveEditor      The curve editor
+	 * @param InCurveModelID         The curve model id for these key handles
 	 * @param InKeyHandles           Array of key handles to create edit objects for
 	 * @param OutObjects             (Out) Array to receive objects that should be used to edit each of the input key handles.
 	 */
+	virtual void CreateKeyProxies(TWeakPtr<FCurveEditor> InWeakCurveEditor, FCurveModelID InCurveModelID, TArrayView<const FKeyHandle> InKeyHandles, TArrayView<UObject*> OutObjects)
+	{}
+
+	UE_DEPRECATED(5.7, "Implement the version that takes a curve editor and curve model ID")
 	virtual void CreateKeyProxies(TArrayView<const FKeyHandle> InKeyHandles, TArrayView<UObject*> OutObjects)
 	{}
 
diff --git a/Engine/Source/Editor/CurveEditor/Public/RichCurveEditorModel.h b/Engine/Source/Editor/CurveEditor/Public/RichCurveEditorModel.h
--- a/Engine/Source/Editor/CurveEditor/Public/RichCurveEditorModel.h
+++ b/Engine/Source/Editor/CurveEditor/Public/RichCurveEditorModel.h
@@ -62,7 +62,7 @@ public:
 	UE_API virtual void AddKeys(TArrayView<const FKeyPosition> InKeyPositions, TArrayView<const FKeyAttributes> InAttributes, TArrayView<TOptional<FKeyHandle>>* OutKeyHandles) override;
 	UE_API virtual void RemoveKeys(TArrayView<const FKeyHandle> InKeys, double InCurrentTime) override;
 
-	UE_API virtual void CreateKeyProxies(TArrayView<const FKeyHandle> InKeyHandles, TArrayView<UObject*> OutObjects) override;
+	UE_API virtual void CreateKeyProxies(TWeakPtr<FCurveEditor> InWeakCurveEditor, FCurveModelID InCurveModelID, TArrayView<const FKeyHandle> InKeyHandles, TArrayView<UObject*> OutObjects) override;
 
 	UE_API virtual TUniquePtr<IBufferedCurveModel> CreateBufferedCurveCopy() const override;
 
diff --git a/Engine/Source/Editor/MovieSceneTools/Private/Channels/BoolChannelCurveModel.cpp b/Engine/Source/Editor/MovieSceneTools/Private/Channels/BoolChannelCurveModel.cpp
--- a/Engine/Source/Editor/MovieSceneTools/Private/Channels/BoolChannelCurveModel.cpp
+++ b/Engine/Source/Editor/MovieSceneTools/Private/Channels/BoolChannelCurveModel.cpp
@@ -91,7 +91,7 @@ FBoolChannelCurveModel::FBoolChannelCurveModel(TMovieSceneChannelHandle<FMovieSc
 {
 }
 
-void FBoolChannelCurveModel::CreateKeyProxies(TArrayView<const FKeyHandle> InKeyHandles, TArrayView<UObject*> OutObjects)
+void FBoolChannelCurveModel::CreateKeyProxies(TWeakPtr<FCurveEditor> InWeakCurveEditor, FCurveModelID InCurveModelID, TArrayView<const FKeyHandle> InKeyHandles, TArrayView<UObject*> OutObjects)
 {
 	for (int32 Index = 0; Index < InKeyHandles.Num(); ++Index)
 	{
diff --git a/Engine/Source/Editor/MovieSceneTools/Private/Channels/ByteChannelCurveModel.cpp b/Engine/Source/Editor/MovieSceneTools/Private/Channels/ByteChannelCurveModel.cpp
--- a/Engine/Source/Editor/MovieSceneTools/Private/Channels/ByteChannelCurveModel.cpp
+++ b/Engine/Source/Editor/MovieSceneTools/Private/Channels/ByteChannelCurveModel.cpp
@@ -91,7 +91,7 @@ FByteChannelCurveModel::FByteChannelCurveModel(TMovieSceneChannelHandle<FMovieSc
 {
 }
 
-void FByteChannelCurveModel::CreateKeyProxies(TArrayView<const FKeyHandle> InKeyHandles, TArrayView<UObject*> OutObjects)
+void FByteChannelCurveModel::CreateKeyProxies(TWeakPtr<FCurveEditor> InWeakCurveEditor, FCurveModelID InCurveModelID, TArrayView<const FKeyHandle> InKeyHandles, TArrayView<UObject*> OutObjects)
 {
 	for (int32 Index = 0; Index < InKeyHandles.Num(); ++Index)
 	{
diff --git a/Engine/Source/Editor/MovieSceneTools/Private/Channels/ConstraintChannelCurveModel.cpp b/Engine/Source/Editor/MovieSceneTools/Private/Channels/ConstraintChannelCurveModel.cpp
--- a/Engine/Source/Editor/MovieSceneTools/Private/Channels/ConstraintChannelCurveModel.cpp
+++ b/Engine/Source/Editor/MovieSceneTools/Private/Channels/ConstraintChannelCurveModel.cpp
@@ -291,7 +291,7 @@ void FConstraintChannelCurveModel::SetCurveAttributes(const FCurveAttributes& In
 {
 }
 
-void FConstraintChannelCurveModel::CreateKeyProxies(TArrayView<const FKeyHandle> InKeyHandles, TArrayView<UObject*> OutObjects)
+void FConstraintChannelCurveModel::CreateKeyProxies(TWeakPtr<FCurveEditor> InWeakCurveEditor, FCurveModelID InCurveModelID, TArrayView<const FKeyHandle> InKeyHandles, TArrayView<UObject*> OutObjects)
 {
 }
 
diff --git a/Engine/Source/Editor/MovieSceneTools/Private/Channels/DoubleChannelCurveModel.cpp b/Engine/Source/Editor/MovieSceneTools/Private/Channels/DoubleChannelCurveModel.cpp
--- a/Engine/Source/Editor/MovieSceneTools/Private/Channels/DoubleChannelCurveModel.cpp
+++ b/Engine/Source/Editor/MovieSceneTools/Private/Channels/DoubleChannelCurveModel.cpp
@@ -131,7 +131,7 @@ void FDoubleChannelCurveModel::GetValueRange(double InMinTime, double InMaxTime,
 	}
 }
 
-void FDoubleChannelCurveModel::CreateKeyProxies(TArrayView<const FKeyHandle> InKeyHandles, TArrayView<UObject*> OutObjects)
+void FDoubleChannelCurveModel::CreateKeyProxies(TWeakPtr<FCurveEditor> InWeakCurveEditor, FCurveModelID InCurveModelID, TArrayView<const FKeyHandle> InKeyHandles, TArrayView<UObject*> OutObjects)
 {
 	for (int32 Index = 0; Index < InKeyHandles.Num(); ++Index)
 	{
diff --git a/Engine/Source/Editor/MovieSceneTools/Private/Channels/EventChannelCurveModel.cpp b/Engine/Source/Editor/MovieSceneTools/Private/Channels/EventChannelCurveModel.cpp
--- a/Engine/Source/Editor/MovieSceneTools/Private/Channels/EventChannelCurveModel.cpp
+++ b/Engine/Source/Editor/MovieSceneTools/Private/Channels/EventChannelCurveModel.cpp
@@ -225,7 +225,7 @@ void FEventChannelCurveModel::SetCurveAttributes(const FCurveAttributes& InCurve
 {
 }
 
-void FEventChannelCurveModel::CreateKeyProxies(TArrayView<const FKeyHandle> InKeyHandles, TArrayView<UObject*> OutObjects)
+void FEventChannelCurveModel::CreateKeyProxies(TWeakPtr<FCurveEditor> InWeakCurveEditor, FCurveModelID InCurveModelID, TArrayView<const FKeyHandle> InKeyHandles, TArrayView<UObject*> OutObjects)
 {
 }
 
diff --git a/Engine/Source/Editor/MovieSceneTools/Private/Channels/EventChannelCurveModel.h b/Engine/Source/Editor/MovieSceneTools/Private/Channels/EventChannelCurveModel.h
--- a/Engine/Source/Editor/MovieSceneTools/Private/Channels/EventChannelCurveModel.h
+++ b/Engine/Source/Editor/MovieSceneTools/Private/Channels/EventChannelCurveModel.h
@@ -58,7 +58,7 @@ public:
 	virtual void AddKeys(TArrayView<const FKeyPosition> InKeyPositions, TArrayView<const FKeyAttributes> InAttributes, TArrayView<TOptional<FKeyHandle>>* OutKeyHandles) override;
 	virtual void RemoveKeys(TArrayView<const FKeyHandle> InKeys, double InCurrentTime) override;
 
-	virtual void CreateKeyProxies(TArrayView<const FKeyHandle> InKeyHandles, TArrayView<UObject*> OutObjects) override;
+	virtual void CreateKeyProxies(TWeakPtr<FCurveEditor> InWeakCurveEditor, FCurveModelID InCurveModelID, TArrayView<const FKeyHandle> InKeyHandles, TArrayView<UObject*> OutObjects) override;
 
 private:
 
diff --git a/Engine/Source/Editor/MovieSceneTools/Private/Channels/FloatChannelCurveModel.cpp b/Engine/Source/Editor/MovieSceneTools/Private/Channels/FloatChannelCurveModel.cpp
--- a/Engine/Source/Editor/MovieSceneTools/Private/Channels/FloatChannelCurveModel.cpp
+++ b/Engine/Source/Editor/MovieSceneTools/Private/Channels/FloatChannelCurveModel.cpp
@@ -79,7 +79,7 @@ FFloatChannelCurveModel::FFloatChannelCurveModel(TMovieSceneChannelHandle<FMovie
 {
 }
 
-void FFloatChannelCurveModel::CreateKeyProxies(TArrayView<const FKeyHandle> InKeyHandles, TArrayView<UObject*> OutObjects)
+void FFloatChannelCurveModel::CreateKeyProxies(TWeakPtr<FCurveEditor> InWeakCurveEditor, FCurveModelID InCurveModelID, TArrayView<const FKeyHandle> InKeyHandles, TArrayView<UObject*> OutObjects)
 {
 	for (int32 Index = 0; Index < InKeyHandles.Num(); ++Index)
 	{
diff --git a/Engine/Source/Editor/MovieSceneTools/Private/Channels/IntegerChannelCurveModel.cpp b/Engine/Source/Editor/MovieSceneTools/Private/Channels/IntegerChannelCurveModel.cpp
--- a/Engine/Source/Editor/MovieSceneTools/Private/Channels/IntegerChannelCurveModel.cpp
+++ b/Engine/Source/Editor/MovieSceneTools/Private/Channels/IntegerChannelCurveModel.cpp
@@ -162,7 +162,7 @@ void FIntegerChannelCurveModel::DrawCurve(const FCurveEditor& CurveEditor, const
 	}
 }
 
-void FIntegerChannelCurveModel::CreateKeyProxies(TArrayView<const FKeyHandle> InKeyHandles, TArrayView<UObject*> OutObjects)
+void FIntegerChannelCurveModel::CreateKeyProxies(TWeakPtr<FCurveEditor> InWeakCurveEditor, FCurveModelID InCurveModelID, TArrayView<const FKeyHandle> InKeyHandles, TArrayView<UObject*> OutObjects)
 {
 	for (int32 Index = 0; Index < InKeyHandles.Num(); ++Index)
 	{
diff --git a/Engine/Source/Editor/MovieSceneTools/Public/Channels/BoolChannelCurveModel.h b/Engine/Source/Editor/MovieSceneTools/Public/Channels/BoolChannelCurveModel.h
--- a/Engine/Source/Editor/MovieSceneTools/Public/Channels/BoolChannelCurveModel.h
+++ b/Engine/Source/Editor/MovieSceneTools/Public/Channels/BoolChannelCurveModel.h
@@ -28,7 +28,7 @@ public:
 	FBoolChannelCurveModel(TMovieSceneChannelHandle<FMovieSceneBoolChannel> InChannel, UMovieSceneSection* InOwningSection, TWeakPtr<ISequencer> InWeakSequencer);
 
 	// FCurveModel
-	virtual void CreateKeyProxies(TArrayView<const FKeyHandle> InKeyHandles, TArrayView<UObject*> OutObjects) override;
+	virtual void CreateKeyProxies(TWeakPtr<FCurveEditor> InWeakCurveEditor, FCurveModelID InCurveModelID, TArrayView<const FKeyHandle> InKeyHandles, TArrayView<UObject*> OutObjects) override;
 	virtual TUniquePtr<IBufferedCurveModel> CreateBufferedCurveCopy() const override;
 	virtual void GetCurveAttributes(FCurveAttributes& OutAttributes) const override;
 	virtual void SetCurveAttributes(const FCurveAttributes& InAttributes) override;
diff --git a/Engine/Source/Editor/MovieSceneTools/Public/Channels/ByteChannelCurveModel.h b/Engine/Source/Editor/MovieSceneTools/Public/Channels/ByteChannelCurveModel.h
--- a/Engine/Source/Editor/MovieSceneTools/Public/Channels/ByteChannelCurveModel.h
+++ b/Engine/Source/Editor/MovieSceneTools/Public/Channels/ByteChannelCurveModel.h
@@ -24,7 +24,7 @@ public:
 	FByteChannelCurveModel(TMovieSceneChannelHandle<FMovieSceneByteChannel> InChannel, UMovieSceneSection* InOwningSection, TWeakPtr<ISequencer> InWeakSequencer);
 
 	// FCurveModel
-	virtual void CreateKeyProxies(TArrayView<const FKeyHandle> InKeyHandles, TArrayView<UObject*> OutObjects) override;
+	virtual void CreateKeyProxies(TWeakPtr<FCurveEditor> InWeakCurveEditor, FCurveModelID InCurveModelID, TArrayView<const FKeyHandle> InKeyHandles, TArrayView<UObject*> OutObjects) override;
 	virtual TUniquePtr<IBufferedCurveModel> CreateBufferedCurveCopy() const override;
 	virtual void GetCurveAttributes(FCurveAttributes& OutAttributes) const override;
 	virtual void SetCurveAttributes(const FCurveAttributes& InAttributes) override;
diff --git a/Engine/Source/Editor/MovieSceneTools/Public/Channels/ConstraintChannelCurveModel.h b/Engine/Source/Editor/MovieSceneTools/Public/Channels/ConstraintChannelCurveModel.h
--- a/Engine/Source/Editor/MovieSceneTools/Public/Channels/ConstraintChannelCurveModel.h
+++ b/Engine/Source/Editor/MovieSceneTools/Public/Channels/ConstraintChannelCurveModel.h
@@ -64,7 +64,7 @@ public:
 	virtual void AddKeys(TArrayView<const FKeyPosition> InKeyPositions, TArrayView<const FKeyAttributes> InAttributes, TArrayView<TOptional<FKeyHandle>>* OutKeyHandles) override;
 	virtual void RemoveKeys(TArrayView<const FKeyHandle> InKeys, double InCurrentTime) override;
 
-	virtual void CreateKeyProxies(TArrayView<const FKeyHandle> InKeyHandles, TArrayView<UObject*> OutObjects) override;
+	virtual void CreateKeyProxies(TWeakPtr<FCurveEditor> InWeakCurveEditor, FCurveModelID InCurveModelID, TArrayView<const FKeyHandle> InKeyHandles, TArrayView<UObject*> OutObjects) override;
 	
 	virtual UObject* GetOwningObject() const override;
 	virtual bool HasChangedAndResetTest() override;
diff --git a/Engine/Source/Editor/MovieSceneTools/Public/Channels/DoubleChannelCurveModel.h b/Engine/Source/Editor/MovieSceneTools/Public/Channels/DoubleChannelCurveModel.h
--- a/Engine/Source/Editor/MovieSceneTools/Public/Channels/DoubleChannelCurveModel.h
+++ b/Engine/Source/Editor/MovieSceneTools/Public/Channels/DoubleChannelCurveModel.h
@@ -24,7 +24,7 @@ public:
 	FDoubleChannelCurveModel(TMovieSceneChannelHandle<FMovieSceneDoubleChannel> InChannel, UMovieSceneSection* InOwningSection, UObject* InOwningObject, TWeakPtr<ISequencer> InWeakSequencer);
 
 	// FCurveModel
-	virtual void CreateKeyProxies(TArrayView<const FKeyHandle> InKeyHandles, TArrayView<UObject*> OutObjects) override;
+	virtual void CreateKeyProxies(TWeakPtr<FCurveEditor> InWeakCurveEditor, FCurveModelID InCurveModelID, TArrayView<const FKeyHandle> InKeyHandles, TArrayView<UObject*> OutObjects) override;
 	virtual TUniquePtr<IBufferedCurveModel> CreateBufferedCurveCopy() const override;
 
 	virtual void GetValueRange(double& MinValue, double& MaxValue) const override;
diff --git a/Engine/Source/Editor/MovieSceneTools/Public/Channels/FloatChannelCurveModel.h b/Engine/Source/Editor/MovieSceneTools/Public/Channels/FloatChannelCurveModel.h
--- a/Engine/Source/Editor/MovieSceneTools/Public/Channels/FloatChannelCurveModel.h
+++ b/Engine/Source/Editor/MovieSceneTools/Public/Channels/FloatChannelCurveModel.h
@@ -23,6 +23,6 @@ public:
 	FFloatChannelCurveModel(TMovieSceneChannelHandle<FMovieSceneFloatChannel> InChannel, UMovieSceneSection* InOwningSection, TWeakPtr<ISequencer> InWeakSequencer);
 
 	// FCurveModel
-	virtual void CreateKeyProxies(TArrayView<const FKeyHandle> InKeyHandles, TArrayView<UObject*> OutObjects) override;
+	virtual void CreateKeyProxies(TWeakPtr<FCurveEditor> InWeakCurveEditor, FCurveModelID InCurveModelID, TArrayView<const FKeyHandle> InKeyHandles, TArrayView<UObject*> OutObjects) override;
 	virtual TUniquePtr<IBufferedCurveModel> CreateBufferedCurveCopy() const override;
 };
\ No newline at end of file
diff --git a/Engine/Source/Editor/MovieSceneTools/Public/Channels/IntegerChannelCurveModel.h b/Engine/Source/Editor/MovieSceneTools/Public/Channels/IntegerChannelCurveModel.h
--- a/Engine/Source/Editor/MovieSceneTools/Public/Channels/IntegerChannelCurveModel.h
+++ b/Engine/Source/Editor/MovieSceneTools/Public/Channels/IntegerChannelCurveModel.h
@@ -29,7 +29,7 @@ public:
 
 	// FCurveModel
 	virtual void DrawCurve(const FCurveEditor& CurveEditor, const FCurveEditorScreenSpace& ScreenSpace, TArray<TTuple<double, double>>& InterpolatingPoints) const override;
-	virtual void CreateKeyProxies(TArrayView<const FKeyHandle> InKeyHandles, TArrayView<UObject*> OutObjects) override;
+	virtual void CreateKeyProxies(TWeakPtr<FCurveEditor> InWeakCurveEditor, FCurveModelID InCurveModelID, TArrayView<const FKeyHandle> InKeyHandles, TArrayView<UObject*> OutObjects) override;
 	virtual TUniquePtr<IBufferedCurveModel> CreateBufferedCurveCopy() const override;
 	virtual void GetCurveAttributes(FCurveAttributes& OutAttributes) const override;
 	virtual void SetCurveAttributes(const FCurveAttributes& InAttributes) override;
