Curve Editor: Added back the semi constrained zoom and removed constrained pan altogether

Follow up to 43884100

diff --git a/Engine/Source/Editor/CurveEditor/Private/DragOperations/CurveEditorDragOperation_Pan.cpp b/Engine/Source/Editor/CurveEditor/Private/DragOperations/CurveEditorDragOperation_Pan.cpp
--- a/Engine/Source/Editor/CurveEditor/Private/DragOperations/CurveEditorDragOperation_Pan.cpp
+++ b/Engine/Source/Editor/CurveEditor/Private/DragOperations/CurveEditorDragOperation_Pan.cpp
@@ -27,15 +27,11 @@ void FCurveEditorDragOperation_PanView::OnBeginDrag(FVector2D InitialPosition, F
 	InitialOutputMax = ViewSpace.GetOutputMax();
 
 	bIsDragging = true;
-
-	SnappingState.Reset();
-	LastMousePosition = CurrentPosition;
 }
 
 void FCurveEditorDragOperation_PanView::OnDrag(FVector2D InitialPosition, FVector2D CurrentPosition, const FPointerEvent& MouseEvent)
 {
-	FVector2D PixelDelta = CurveEditor->GetAxisSnap().GetSnappedPosition(InitialPosition, LastMousePosition,CurrentPosition, MouseEvent, SnappingState, true) - InitialPosition;
-	LastMousePosition = CurrentPosition;
+	const FVector2D PixelDelta = CurrentPosition - InitialPosition;
 
 	FCurveEditorScreenSpace ViewSpace = View->GetViewSpace();
 
@@ -84,15 +80,12 @@ void FCurveEditorDragOperation_PanInput::OnBeginDrag(FVector2D InitialPosition,
 	FCurveEditorScreenSpaceH InputSpace = CurveEditor->GetPanelInputSpace();
 	InitialInputMin = InputSpace.GetInputMin();
 	InitialInputMax = InputSpace.GetInputMax();
-	SnappingState.Reset();
-	LastMousePosition = CurrentPosition;
 
 }
 
 void FCurveEditorDragOperation_PanInput::OnDrag(FVector2D InitialPosition, FVector2D CurrentPosition, const FPointerEvent& MouseEvent)
 {
-	FVector2D PixelDelta = CurveEditor->GetAxisSnap().GetSnappedPosition(InitialPosition, LastMousePosition, CurrentPosition, MouseEvent, SnappingState, true) - InitialPosition;
-	LastMousePosition = CurrentPosition;
+	const FVector2D PixelDelta = CurrentPosition - InitialPosition;
 
 	FCurveEditorScreenSpaceH InputSpace = CurveEditor->GetPanelInputSpace();
 
diff --git a/Engine/Source/Editor/CurveEditor/Private/DragOperations/CurveEditorDragOperation_Pan.h b/Engine/Source/Editor/CurveEditor/Private/DragOperations/CurveEditorDragOperation_Pan.h
--- a/Engine/Source/Editor/CurveEditor/Private/DragOperations/CurveEditorDragOperation_Pan.h
+++ b/Engine/Source/Editor/CurveEditor/Private/DragOperations/CurveEditorDragOperation_Pan.h
@@ -47,6 +47,4 @@ private:
 
 	FCurveEditor* CurveEditor;
 	double InitialInputMin, InitialInputMax;
-	FCurveEditorAxisSnap::FSnapState SnappingState;
-	FVector2D LastMousePosition;
 };
\ No newline at end of file
diff --git a/Engine/Source/Editor/CurveEditor/Private/DragOperations/CurveEditorDragOperation_Zoom.cpp b/Engine/Source/Editor/CurveEditor/Private/DragOperations/CurveEditorDragOperation_Zoom.cpp
--- a/Engine/Source/Editor/CurveEditor/Private/DragOperations/CurveEditorDragOperation_Zoom.cpp
+++ b/Engine/Source/Editor/CurveEditor/Private/DragOperations/CurveEditorDragOperation_Zoom.cpp
@@ -43,17 +43,13 @@ void FCurveEditorDragOperation_Zoom::OnBeginDrag(FVector2D InitialPosition, FVec
 		OriginalOutputRange = (OutputSpace.GetOutputMax() - OutputSpace.GetOutputMin());
 		ZoomOriginY = OutputSpace.GetOutputMin() + OriginalOutputRange * (1.f - ZoomFactor.Y);
 	}
-
-	SnappingState.Reset();
-	LastMousePosition = CurrentPosition;
 }
 
 void FCurveEditorDragOperation_Zoom::OnDrag(FVector2D InitialPosition, FVector2D CurrentPosition, const FPointerEvent& MouseEvent)
 {
 	constexpr double ClampRange = 1e9f;
 
-	const FVector2D PixelDelta = CurveEditor->GetAxisSnap().GetSnappedPosition(InitialPosition, LastMousePosition, CurrentPosition, MouseEvent, SnappingState, true) - InitialPosition;
-	LastMousePosition = CurrentPosition;
+	const FVector2D PixelDelta = CurrentPosition - InitialPosition;
 
 	const FCurveEditorZoomScaleConfig& ZoomConfig = CurveEditor->GetZoomScaleConfig();
 
diff --git a/Engine/Source/Editor/CurveEditor/Private/DragOperations/CurveEditorDragOperation_Zoom.h b/Engine/Source/Editor/CurveEditor/Private/DragOperations/CurveEditorDragOperation_Zoom.h
--- a/Engine/Source/Editor/CurveEditor/Private/DragOperations/CurveEditorDragOperation_Zoom.h
+++ b/Engine/Source/Editor/CurveEditor/Private/DragOperations/CurveEditorDragOperation_Zoom.h
@@ -29,9 +29,6 @@ private:
 	double ZoomOriginX, ZoomOriginY;
 	double OriginalInputRange, OriginalOutputRange;
 
-	FCurveEditorAxisSnap::FSnapState SnappingState;
-	FVector2D LastMousePosition;
-
 	double GetZoomMultiplier_InputAxis(double InMovedMouseX) const;
 	double GetZoomMultiplier_OutputAxis(double InMovedMouseY) const;
 };
