Curve Editor: Introduce FScopedKeyChange - the workhorse of command-based undo for changing keys (move, add, delete, attributes).

Whenever you want to make an undoable action in Curve Editor which changes curves, all you need to is construct a FScopedKeyChange (specify the curves you and the types of changes to track), make the change, and then let the instance go out of scope.

This change only adds the architecture - the existing undo in Curve Editor remains to be diff-based. A follow up change will migrate to this new API.

diff --git a/Engine/Source/Editor/CurveEditor/CurveEditor.Build.cs b/Engine/Source/Editor/CurveEditor/CurveEditor.Build.cs
--- a/Engine/Source/Editor/CurveEditor/CurveEditor.Build.cs
+++ b/Engine/Source/Editor/CurveEditor/CurveEditor.Build.cs
@@ -22,6 +22,7 @@ public class CurveEditor : ModuleRules
 				"UMG",
 				"UnrealEd",
 				"ToolMenus",
+				"TraceLog"
 			}
 		);
 
diff --git a/Engine/Source/Editor/CurveEditor/Private/CurveEditorTrace.cpp b/Engine/Source/Editor/CurveEditor/Private/CurveEditorTrace.cpp
new file mode 100644
--- /dev/null
+++ b/Engine/Source/Editor/CurveEditor/Private/CurveEditorTrace.cpp
@@ -0,0 +1,7 @@
+﻿// Copyright Epic Games, Inc. All Rights Reserved.
+
+#include "CurveEditorTrace.h"
+
+#include "Trace/Detail/Channel.h"
+
+UE_TRACE_CHANNEL_DEFINE(CurveEditorChannel);
diff --git a/Engine/Source/Editor/CurveEditor/Private/CurveEditorTrace.h b/Engine/Source/Editor/CurveEditor/Private/CurveEditorTrace.h
new file mode 100644
--- /dev/null
+++ b/Engine/Source/Editor/CurveEditor/Private/CurveEditorTrace.h
@@ -0,0 +1,8 @@
+﻿// Copyright Epic Games, Inc. All Rights Reserved.
+
+#pragma once
+
+#include "Trace/Trace.h"
+
+UE_TRACE_CHANNEL_EXTERN(CurveEditorChannel);
+#define SCOPED_CURVE_EDITOR_TRACE(TraceName) TRACE_CPUPROFILER_EVENT_SCOPE_ON_CHANNEL(TraceName, CurveEditorChannel)
\ No newline at end of file
diff --git a/Engine/Source/Editor/CurveEditor/Private/Modification/Keys/CurveAttributeChangeUtils.cpp b/Engine/Source/Editor/CurveEditor/Private/Modification/Keys/CurveAttributeChangeUtils.cpp
new file mode 100644
--- /dev/null
+++ b/Engine/Source/Editor/CurveEditor/Private/Modification/Keys/CurveAttributeChangeUtils.cpp
@@ -0,0 +1,36 @@
+﻿// Copyright Epic Games, Inc. All Rights Reserved.
+
+#include "CurveAttributeChangeUtils.h"
+
+#include "CurveModel.h"
+#include "Modification/Keys/Data/CurveAttributeChangeData.h"
+
+namespace UE::CurveEditor::CurveAttributes
+{
+FCurveAttributeChangeData_PerCurve Diff(const FCurveAttributes& InOriginal, const FCurveAttributes& InTarget)
+{
+	return FCurveAttributeChangeData_PerCurve{ InOriginal, InTarget };
+}
+
+void ApplyChange(const TMap<FCurveModelID, TUniquePtr<FCurveModel>>& InCurves, const FCurveAttributeChangeData& InDeltaChange)
+{
+	for (const TPair<FCurveModelID, TUniquePtr<FCurveModel>>& Pair : InCurves)
+	{
+		if (const FCurveAttributeChangeData_PerCurve* Change = InDeltaChange.ChangeData.Find(Pair.Key))
+		{
+			Pair.Value->SetCurveAttributes(Change->AfterChange);
+		}
+	}
+}
+
+void RevertChange(const TMap<FCurveModelID, TUniquePtr<FCurveModel>>& InCurves, const FCurveAttributeChangeData& InDeltaChange)
+{
+	for (const TPair<FCurveModelID, TUniquePtr<FCurveModel>>& Pair : InCurves)
+	{
+		if (const FCurveAttributeChangeData_PerCurve* Change = InDeltaChange.ChangeData.Find(Pair.Key))
+		{
+			Pair.Value->SetCurveAttributes(Change->BeforeChange);
+		}
+	}
+}
+}
diff --git a/Engine/Source/Editor/CurveEditor/Private/Modification/Keys/CurveAttributeChangeUtils.h b/Engine/Source/Editor/CurveEditor/Private/Modification/Keys/CurveAttributeChangeUtils.h
new file mode 100644
--- /dev/null
+++ b/Engine/Source/Editor/CurveEditor/Private/Modification/Keys/CurveAttributeChangeUtils.h
@@ -0,0 +1,28 @@
+﻿// Copyright Epic Games, Inc. All Rights Reserved.
+
+#pragma once
+
+#include "Containers/ContainersFwd.h"
+#include "Templates/UniquePtr.h"
+
+class FCurveModel;
+namespace UE::CurveEditor { struct FCurveAttributeChangeData; }
+namespace UE::CurveEditor { struct FCurveAttributeChangeData_PerCurve; }
+struct FCurveAttributes;
+struct FCurveModelID;
+
+namespace UE::CurveEditor::CurveAttributes
+{
+/**
+ * Computes the delta change to get from InOriginal to InTarget.
+ * 
+ * @param InOriginal The original curve attributes.
+ * @param InTarget The new curve attributes.
+ */
+FCurveAttributeChangeData_PerCurve Diff(const FCurveAttributes& InOriginal, const FCurveAttributes& InTarget);
+
+/** Applies InDeltaChange to InCurves (redo). */
+void ApplyChange(const TMap<FCurveModelID, TUniquePtr<FCurveModel>>& InCurves, const FCurveAttributeChangeData& InDeltaChange);
+/** Reverts InDeltaChange from InCurves (undo). */
+void RevertChange(const TMap<FCurveModelID, TUniquePtr<FCurveModel>>& InCurves, const FCurveAttributeChangeData& InDeltaChange);
+}
diff --git a/Engine/Source/Editor/CurveEditor/Private/Modification/Keys/CurveChangeDiff.cpp b/Engine/Source/Editor/CurveEditor/Private/Modification/Keys/CurveChangeDiff.cpp
new file mode 100644
--- /dev/null
+++ b/Engine/Source/Editor/CurveEditor/Private/Modification/Keys/CurveChangeDiff.cpp
@@ -0,0 +1,226 @@
+﻿// Copyright Epic Games, Inc. All Rights Reserved.
+
+#include "Modification/Keys/CurveChangeDiff.h"
+
+#include "CurveAttributeChangeUtils.h"
+#include "CurveEditor.h"
+#include "CurveEditorTrace.h"
+#include "CurveModel.h"
+#include "KeyAttributeChangeUtils.h"
+#include "KeyExistenceUtils.h"
+#include "Modification/Keys/Data/CurveKeyDataMap.h"
+#include "Modification/Keys/Data/GenericCurveChangeData.h"
+#include "Modification/Keys/ECurveChangeFlags.h"
+#include "Modification/Keys/MoveKeysChangeUtils.h"
+#include "SCurveEditor.h"
+
+namespace UE::CurveEditor
+{
+namespace Private
+{
+static void HandleOp_MoveKeys(
+	const ECurveChangeFlags Operations, const FCurveModelID& CurveId, const FCurveModel& CurveModel, const FCurveKeyData& SavedState,
+	FGenericCurveChangeData& OutChange
+	)
+{
+	if (!EnumHasAnyFlags(Operations, ECurveChangeFlags::MoveKeys) || !SavedState || !ensureAlways(SavedState.HasKeyPositions()))
+	{
+		return;
+	}
+
+	TArray<FKeyPosition> CurrentPositions;
+	CurrentPositions.SetNum(SavedState.KeyHandles.Num());
+	CurveModel.GetKeyPositions(SavedState.KeyHandles, CurrentPositions);
+			
+	if (FMoveKeysChangeData_PerCurve PerCurveData = MoveKeys::Diff(SavedState.KeyHandles, SavedState.KeyPositions, CurrentPositions))
+	{
+		OutChange.MoveKeysData.ChangedCurves.Emplace(CurveId, MoveTemp(PerCurveData));
+	}
+}
+
+static void HandleOp_AddedKeys(
+	const ECurveChangeFlags Operations, const FCurveModelID& CurveId, const FCurveModel& CurveModel, const FCurveKeyData& SavedState,
+	FGenericCurveChangeData& OutChange
+)
+{
+	if (!EnumHasAnyFlags(Operations, ECurveChangeFlags::AddKeys) || !SavedState)
+	{
+		return;
+	}
+
+	if (FCurveKeyData AddedKeys = KeyInsertion::Diff(SavedState.KeyHandles, CurveModel))
+	{
+		OutChange.AddKeysData.SavedCurveState.Emplace(CurveId, MoveTemp(AddedKeys));
+	}
+}
+
+static void HandleOp_RemovedKeys(
+	const ECurveChangeFlags Operations, const FCurveModelID& CurveId, const FCurveModel& CurveModel, const FCurveKeyData& SavedState,
+	FGenericCurveChangeData& OutChange
+)
+{
+	if (!EnumHasAnyFlags(Operations, ECurveChangeFlags::RemoveKeys) || !SavedState
+		|| !ensureAlways(SavedState.HasKeyPositions()) || !ensureAlways(SavedState.HasKeyAttributes()))
+	{
+		return;
+	}
+
+	if (FCurveKeyData Removed = KeyRemoval::Diff(SavedState.KeyHandles, SavedState.KeyPositions, SavedState.KeyAttributes, CurveModel))
+	{
+		OutChange.RemoveKeysData.SavedCurveState.Emplace(CurveId, MoveTemp(Removed));
+	}
+}
+	
+static void HandleOp_KeyAttributes(
+	const ECurveChangeFlags Operations, const FCurveModelID& CurveId, const FCurveModel& CurveModel, const FCurveKeyData& SavedState,
+	FGenericCurveChangeData& OutChange
+)
+{
+	if (!EnumHasAnyFlags(Operations, ECurveChangeFlags::KeyAttributes) || !SavedState || !ensureAlways(SavedState.HasKeyAttributes()))
+	{
+		return;
+	}
+
+	TArray<FKeyAttributes> CurrentAttributes;
+	CurrentAttributes.SetNum(SavedState.KeyHandles.Num());
+	CurveModel.GetKeyAttributesExcludingAutoComputed(SavedState.KeyHandles, CurrentAttributes);
+
+	if (FKeyAttributeChangeData_PerCurve PerCurveData = KeyAttributes::Diff(SavedState.KeyHandles, SavedState.KeyAttributes, CurrentAttributes))
+	{
+		OutChange.KeyAttributeData.ChangedCurves.Emplace(CurveId, MoveTemp(PerCurveData));
+	}
+}
+
+static void HandleOp_CurveAttributes(
+	const ECurveChangeFlags Operations, const FCurveModelID& CurveId, const FCurveModel& CurveModel, const FCurveDiffingData& SavedState,
+	FGenericCurveChangeData& OutChange
+)
+{
+	if (!EnumHasAnyFlags(Operations, ECurveChangeFlags::CurveAttributes) || !SavedState || !ensureAlways(SavedState.HasCurveAttributes()))
+	{
+		return;
+	}
+
+	FCurveAttributes CurrentAttributes;
+	CurveModel.GetCurveAttributes(CurrentAttributes);
+	
+	if (FCurveAttributeChangeData_PerCurve PerCurveData = CurveAttributes::Diff(*SavedState.CurveAttributes, CurrentAttributes))
+	{
+		OutChange.CurveAttributeData.ChangeData.Emplace(CurveId, PerCurveData);
+	}
+}
+
+/** Returns a copy of InSavedState minus the keys in InRemovedKeys.  */
+static FCurveKeyData CleanseOfRemovedKeys(
+	TConstArrayView<FKeyHandle> InRemovedKeys, const FCurveKeyData& InSavedState
+	)
+{
+	SCOPED_CURVE_EDITOR_TRACE(CleanseRemovedKeys);
+	const int32 TargetNum =  InSavedState.KeyHandles.Num() - InRemovedKeys.Num();
+	if (TargetNum <= 0)
+	{
+		return {};
+	}
+	
+	FCurveKeyData Result;
+	Result.KeyHandles.Reserve(TargetNum);
+	const bool bDoAttributes = InSavedState.HasKeyAttributes();
+	if (bDoAttributes)
+	{
+		Result.KeyAttributes.Reserve(TargetNum);
+	}
+	const bool bDoPositions = InSavedState.HasKeyPositions();
+	if (bDoPositions)
+	{
+		Result.KeyPositions.Reserve(TargetNum);
+	} 
+	
+	for (int32 Index = 0; Index < InSavedState.KeyHandles.Num() && Result.KeyHandles.Num() < TargetNum; ++Index)
+	{
+		const FKeyHandle KeyHandle = InSavedState.KeyHandles[Index];
+		if (InRemovedKeys.Contains(KeyHandle))
+		{
+			continue;
+		}
+
+		Result.KeyHandles.Add(KeyHandle);
+		if (bDoAttributes)
+		{
+			Result.KeyAttributes.Add(InSavedState.KeyAttributes[Index]);
+		}
+		if (bDoPositions)
+		{
+			Result.KeyPositions.Add(InSavedState.KeyPositions[Index]);
+		}
+	}
+	return Result;
+}
+
+/** Checks whether any keys in InSavedState were removed from CurveModel and if so, returns a copy of the state minus the removed key data. */
+TOptional<FCurveKeyData> CleanseStateIfKeysWereRemoved(
+	const FCurveKeyData& InSavedState, const ECurveChangeFlags Operations, const FRemoveKeysChangeData& InKeyRemovalData,
+	const FCurveModelID& CurveId, const FCurveModel& CurveModel
+	)
+{
+	if (const FCurveKeyData* RemovedKeyData = InKeyRemovalData.SavedCurveState.Find(CurveId))
+	{
+		return CleanseOfRemovedKeys(RemovedKeyData->KeyHandles, InSavedState);
+	}
+
+	// If InKeyRemovalData contained no removal data but RemoveKeys is set, that means we already analyzed removal and don't have to do it again.
+	// Effectively, this implies that no keys have been removed.
+	if (!EnumHasAnyFlags(Operations, ECurveChangeFlags::RemoveKeys))
+	{
+		const TArray<FKeyHandle> RemovedKeys = KeyRemoval::Diff(InSavedState.KeyHandles, CurveModel);
+		return CleanseOfRemovedKeys(RemovedKeys, InSavedState);
+	}
+	
+	return {};
+}
+}
+
+FGenericCurveChangeData ComputeDiff(const FCurveEditor& InCurveEditor, const ECurveChangeFlags InOperations, const FCurvesSnapshot& InSnapshot)
+{
+	SCOPED_CURVE_EDITOR_TRACE(ComputeDiff);
+	
+	FGenericCurveChangeData Change;
+	for (const TPair<FCurveModelID, FCurveDiffingData>& StatePair : InSnapshot.CurveData)
+	{
+		const FCurveModelID& CurveId = StatePair.Key;
+		if (const FCurveModel* CurveModel = InCurveEditor.FindCurve(CurveId))
+		{
+			const FCurveDiffingData& SavedState = StatePair.Value;
+
+			Private::HandleOp_CurveAttributes(InOperations, CurveId, *CurveModel, SavedState, Change);
+			
+			// Do RemoveKeys before MoveKeys and KeyAttributes... 
+			Private::HandleOp_RemovedKeys(InOperations, CurveId, *CurveModel, SavedState, Change);
+			Private::HandleOp_AddedKeys(InOperations, CurveId, *CurveModel, SavedState, Change);
+			
+			// ... so we can reuse that info to tell MoveKeys and KeyAttributes to skip analyzing the keys that were removed.
+			const TOptional<FCurveKeyData> StateMinusRemovedKeys = Private::CleanseStateIfKeysWereRemoved(
+				SavedState, InOperations, Change.RemoveKeysData, CurveId, *CurveModel
+				);
+			if (StateMinusRemovedKeys)
+			{
+				// The operations should ignore the keys that were removed since that data is now irrelevant.
+				// Not ignoring would cause them to report a change when there is none.
+				// Example: If we didn't cleanse keys removed by the operation, HandleOp_MoveKeys would try to call GetKeyPositions on FKeyHandles
+				// that were removed by the operation.
+				// GetKeyPositions would return default-initialized FKeyPositions. It would then proceed to compare the default-initialized state
+				// against the key value from before the change, and conclude that the keys removed by the operation were "moved" (false change).
+				Private::HandleOp_MoveKeys(InOperations, CurveId, *CurveModel, *StateMinusRemovedKeys, Change);
+				Private::HandleOp_KeyAttributes(InOperations, CurveId, *CurveModel, *StateMinusRemovedKeys, Change);
+			}
+			else
+			{
+				// Nothing was removed, so just pass in the same state
+				Private::HandleOp_MoveKeys(InOperations, CurveId, *CurveModel, SavedState, Change);
+				Private::HandleOp_KeyAttributes(InOperations, CurveId, *CurveModel, SavedState, Change);
+			}
+		}
+	}
+	
+	return Change;
+}
+}
\ No newline at end of file
diff --git a/Engine/Source/Editor/CurveEditor/Private/Modification/Keys/CurvesSnapshotBuilder.cpp b/Engine/Source/Editor/CurveEditor/Private/Modification/Keys/CurvesSnapshotBuilder.cpp
new file mode 100644
--- /dev/null
+++ b/Engine/Source/Editor/CurveEditor/Private/Modification/Keys/CurvesSnapshotBuilder.cpp
@@ -0,0 +1,150 @@
+﻿// Copyright Epic Games, Inc. All Rights Reserved.
+
+#include "Modification/Keys/CurvesSnapshotBuilder.h"
+
+#include "CurveEditor.h"
+#include "CurveEditorTrace.h"
+
+namespace UE::CurveEditor
+{
+namespace Private
+{
+static bool CanTrackCurve(const FCurveModel& Model)
+{
+	// Read only curves cannot be modified so it is pointless to track them
+	return !Model.IsReadOnly();
+}
+
+static void AddPositionsAndAttributes(const FCurveModel& InCurveModel, FCurveDiffingData& InSavedState, ECurveChangeFlags InOperations)
+{
+	if (EnumHasAnyFlags(InOperations, ECurveChangeFlags::MoveKeys | ECurveChangeFlags::RemoveKeys))
+	{
+		InSavedState.KeyPositions.SetNum(InSavedState.KeyHandles.Num());
+		InCurveModel.GetKeyPositions(InSavedState.KeyHandles, InSavedState.KeyPositions);
+	}
+
+	if (EnumHasAnyFlags(InOperations, ECurveChangeFlags::KeyAttributes | ECurveChangeFlags::RemoveKeys))
+	{
+		InSavedState.KeyAttributes.SetNum(InSavedState.KeyHandles.Num());
+		InCurveModel.GetKeyAttributesExcludingAutoComputed(InSavedState.KeyHandles, InSavedState.KeyAttributes);
+	}
+
+	if (EnumHasAnyFlags(InOperations, ECurveChangeFlags::CurveAttributes))
+	{
+		InSavedState.CurveAttributes.Emplace();
+		InCurveModel.GetCurveAttributes(*InSavedState.CurveAttributes);
+	}
+}
+}
+
+FCurvesSnapshotBuilder& FCurvesSnapshotBuilder::TrackCurve(const FCurveModelID& InCurveId)
+{
+	SCOPED_CURVE_EDITOR_TRACE(FCurvesSnapshotBuilder::TrackCurve);
+	
+	const TSharedPtr<FCurveEditor> CurveEditorPin = WeakCurveEditor.Pin();
+	const FCurveModel* CurveModel = CurveEditorPin ? CurveEditorPin->FindCurve(InCurveId) : nullptr;
+	if (CurveModel && Private::CanTrackCurve(*CurveModel))
+	{
+		// Allow overwriting of previous data.
+		FCurveDiffingData& SavedState = Snapshot.CurveData.Add(InCurveId);
+		CurveModel->GetAllKeys(SavedState.KeyHandles);
+		Private::AddPositionsAndAttributes(*CurveModel, SavedState, CapturedDataFlags);
+	}
+	return *this;
+}
+
+FCurvesSnapshotBuilder& FCurvesSnapshotBuilder::TrackCurves(TConstArrayView<FCurveModelID> InCurves)
+{
+	SCOPED_CURVE_EDITOR_TRACE(FCurvesSnapshotBuilder::TrackCurves);
+	for (const FCurveModelID& CurveId : InCurves)
+	{
+		TrackCurve(CurveId);
+	}
+	return *this;
+}
+
+FCurvesSnapshotBuilder& FCurvesSnapshotBuilder::TrackCurves(const TSet<FCurveModelID>& InCurves)
+{
+	SCOPED_CURVE_EDITOR_TRACE(FCurvesSnapshotBuilder::TrackCurves);
+	for (const FCurveModelID& CurveId : InCurves)
+	{
+		TrackCurve(CurveId);
+	}
+	return *this;
+}
+
+FCurvesSnapshotBuilder& FCurvesSnapshotBuilder::TrackCurve(const FCurveModelID& InCurveId, TConstArrayView<FKeyHandle> Keys)
+{
+	SCOPED_CURVE_EDITOR_TRACE(FCurvesSnapshotBuilder::TrackCurve);
+	// Add keys needs to compare against all keys to properly detect added keys.
+	if (!ensure(!EnumHasAnyFlags(CapturedDataFlags, ECurveChangeFlags::AddKeys)))
+	{
+		return *this;
+	}
+	
+	const TSharedPtr<FCurveEditor> CurveEditorPin = WeakCurveEditor.Pin();
+	const FCurveModel* CurveModel = CurveEditorPin ? CurveEditorPin->FindCurve(InCurveId) : nullptr;
+	if (CurveModel && Private::CanTrackCurve(*CurveModel))
+	{
+		// Allow overwriting of previous data.
+		FCurveDiffingData& SavedState = Snapshot.CurveData.Add(InCurveId);
+
+		SavedState.KeyHandles = TArray<FKeyHandle>(Keys);
+		Private::AddPositionsAndAttributes(*CurveModel, SavedState, CapturedDataFlags);
+	}
+	
+	return *this;
+}
+
+FCurvesSnapshotBuilder& FCurvesSnapshotBuilder::TrackSelectedKeys()
+{
+	SCOPED_CURVE_EDITOR_TRACE(FCurvesSnapshotBuilder::TrackSelectedKeys);
+	// Add keys needs to compare against all keys to properly detect added keys.
+	if (!ensure(!EnumHasAnyFlags(CapturedDataFlags, ECurveChangeFlags::AddKeys)))
+	{
+		return *this;
+	}
+
+	const TSharedPtr<FCurveEditor> CurveEditorPin = WeakCurveEditor.Pin();
+	if (!ensure(CurveEditorPin))
+	{
+		return *this;
+	}
+
+	const TMap<FCurveModelID, FKeyHandleSet>& SelectionMap = CurveEditorPin->GetSelection().GetAll();
+	for (const TPair<FCurveModelID, FKeyHandleSet>& Pair : SelectionMap)
+	{
+		TrackCurve(Pair.Key, Pair.Value.AsArray());
+	}
+
+	return *this;
+}
+
+FCurvesSnapshotBuilder& FCurvesSnapshotBuilder::TrackSelectedCurves()
+{
+	SCOPED_CURVE_EDITOR_TRACE(FCurvesSnapshotBuilder::TrackSelectedCurves);
+	const TSharedPtr<FCurveEditor> CurveEditorPin = WeakCurveEditor.Pin();
+	if (!ensure(CurveEditorPin))
+	{
+		return *this;
+	}
+
+	const TMap<FCurveModelID, FKeyHandleSet>& SelectionMap = CurveEditorPin->GetSelection().GetAll();
+	for (const TPair<FCurveModelID, FKeyHandleSet>& Pair : SelectionMap)
+	{
+		TrackCurve(Pair.Key);
+	}
+
+	return *this;
+}
+
+FCurvesSnapshotBuilder& FCurvesSnapshotBuilder::TrackAllCurves()
+{
+	SCOPED_CURVE_EDITOR_TRACE(FCurvesSnapshotBuilder::TrackAllCurves);
+	if (const TSharedPtr<FCurveEditor> CurveEditorPin = WeakCurveEditor.Pin())
+	{
+		TrackCurves(CurveEditorPin->GetCurves());
+	}
+	return *this;
+}
+}
diff --git a/Engine/Source/Editor/CurveEditor/Private/Modification/Keys/GenericCurveChangeCommand.cpp b/Engine/Source/Editor/CurveEditor/Private/Modification/Keys/GenericCurveChangeCommand.cpp
new file mode 100644
--- /dev/null
+++ b/Engine/Source/Editor/CurveEditor/Private/Modification/Keys/GenericCurveChangeCommand.cpp
@@ -0,0 +1,25 @@
+﻿// Copyright Epic Games, Inc. All Rights Reserved.
+
+#include "GenericCurveChangeCommand.h"
+
+#include "CurveEditor.h"
+#include "Modification/Keys/GenericCurveChangeUtils.h"
+
+namespace UE::CurveEditor
+{
+void FGenericCurveChangeCommand::Apply(UObject* Object)
+{
+	if (const TSharedPtr<FCurveEditor> CurveEditor = GetCurveEditor(Object))
+	{
+		GenericCurveChange::ApplyChange(CurveEditor->GetCurves(), DeltaChange);
+	}
+}
+
+void FGenericCurveChangeCommand::Revert(UObject* Object)
+{
+	if (const TSharedPtr<FCurveEditor> CurveEditor = GetCurveEditor(Object))
+	{
+		GenericCurveChange::RevertChange(CurveEditor->GetCurves(), DeltaChange);
+	}
+}
+}
diff --git a/Engine/Source/Editor/CurveEditor/Private/Modification/Keys/GenericCurveChangeCommand.h b/Engine/Source/Editor/CurveEditor/Private/Modification/Keys/GenericCurveChangeCommand.h
new file mode 100644
--- /dev/null
+++ b/Engine/Source/Editor/CurveEditor/Private/Modification/Keys/GenericCurveChangeCommand.h
@@ -0,0 +1,34 @@
+﻿// Copyright Epic Games, Inc. All Rights Reserved.
+
+#pragma once
+
+#include "Modification/CurveEditorCommandChange.h"
+#include "Modification/Keys/Data/GenericCurveChangeData.h"
+
+namespace UE::CurveEditor
+{
+/**
+ * Captures any changes to keys including and limited to:
+ * - Move keys,
+ * - Add new keys,
+ * - Remove keys,
+ * - Change key attributes,
+ * - Change curve attributes
+ */
+class FGenericCurveChangeCommand : public FCurveEditorCommandChange
+{
+public:
+
+	explicit FGenericCurveChangeCommand(FGenericCurveChangeData InDeltaChange) : DeltaChange(MoveTemp(InDeltaChange)) {}
+
+	//~ Begin FChange Interface
+	virtual void Apply(UObject* Object) override;
+	virtual void Revert(UObject* Object) override;
+	//~ End FChange Interface
+
+private:
+
+	/** The change made to the selection. */
+	const FGenericCurveChangeData DeltaChange;
+};
+}
diff --git a/Engine/Source/Editor/CurveEditor/Private/Modification/Keys/GenericCurveChangeUtils.cpp b/Engine/Source/Editor/CurveEditor/Private/Modification/Keys/GenericCurveChangeUtils.cpp
new file mode 100644
--- /dev/null
+++ b/Engine/Source/Editor/CurveEditor/Private/Modification/Keys/GenericCurveChangeUtils.cpp
@@ -0,0 +1,56 @@
+﻿// Copyright Epic Games, Inc. All Rights Reserved.
+
+#include "Modification/Keys/GenericCurveChangeUtils.h"
+
+#include "CurveAttributeChangeUtils.h"
+#include "CurveEditor.h"
+#include "CurveEditorTrace.h"
+#include "KeyAttributeChangeUtils.h"
+#include "KeyExistenceUtils.h"
+#include "Modification/Keys/Data/GenericCurveChangeData.h"
+#include "Modification/Keys/MoveKeysChangeUtils.h"
+
+namespace UE::CurveEditor::GenericCurveChange
+{
+void ApplyChange(
+	const TMap<FCurveModelID, TUniquePtr<FCurveModel>>& InCurves, const FGenericCurveChangeData& InDeltaChange, const double InCurrentSliderPos
+	)
+{
+	SCOPED_CURVE_EDITOR_TRACE(ApplyGenericChange);
+	KeyRemoval::ApplyChange(InCurves, InDeltaChange.RemoveKeysData, InCurrentSliderPos);
+	KeyInsertion::ApplyChange(InCurves, InDeltaChange.AddKeysData);
+	MoveKeys::ApplyChange(InCurves, InDeltaChange.MoveKeysData);
+	KeyAttributes::ApplyChange(InCurves, InDeltaChange.KeyAttributeData);
+	CurveAttributes::ApplyChange(InCurves, InDeltaChange.CurveAttributeData);
+}
+
+void RevertChange(
+	const TMap<FCurveModelID, TUniquePtr<FCurveModel>>& InCurves, const FGenericCurveChangeData& InDeltaChange, const double InCurrentSliderPos
+	)
+{
+	SCOPED_CURVE_EDITOR_TRACE(RevertGenericChange);
+	KeyInsertion::RevertChange(InCurves, InDeltaChange.AddKeysData, InCurrentSliderPos);
+	KeyRemoval::RevertChange(InCurves, InDeltaChange.RemoveKeysData);
+	MoveKeys::RevertChange(InCurves, InDeltaChange.MoveKeysData);
+	KeyAttributes::RevertChange(InCurves, InDeltaChange.KeyAttributeData);
+	CurveAttributes::RevertChange(InCurves, InDeltaChange.CurveAttributeData);
+}
+
+void ApplyChange(const FCurveEditor& InCurveEditor, const FGenericCurveChangeData& InDeltaChange)
+{
+	const TSharedPtr<ITimeSliderController> TimeController = InCurveEditor.GetTimeSliderController();
+	const double CurrentTime = TimeController
+		? TimeController->GetTickResolution().AsSeconds(TimeController->GetScrubPosition())
+		: 0.0;
+	ApplyChange(InCurveEditor.GetCurves(), InDeltaChange, CurrentTime);
+}
+	
+void RevertChange(const FCurveEditor& InCurveEditor, const FGenericCurveChangeData& InDeltaChange)
+{
+	const TSharedPtr<ITimeSliderController> TimeController = InCurveEditor.GetTimeSliderController();
+	const double CurrentTime = TimeController
+		? TimeController->GetTickResolution().AsSeconds(TimeController->GetScrubPosition())
+		: 0.0;
+	RevertChange(InCurveEditor.GetCurves(), InDeltaChange, CurrentTime);
+}
+}
diff --git a/Engine/Source/Editor/CurveEditor/Private/Modification/Keys/KeyAttributeChangeUtils.cpp b/Engine/Source/Editor/CurveEditor/Private/Modification/Keys/KeyAttributeChangeUtils.cpp
new file mode 100644
--- /dev/null
+++ b/Engine/Source/Editor/CurveEditor/Private/Modification/Keys/KeyAttributeChangeUtils.cpp
@@ -0,0 +1,77 @@
+﻿// Copyright Epic Games, Inc. All Rights Reserved.
+
+#include "KeyAttributeChangeUtils.h"
+
+#include "CurveEditorTrace.h"
+#include "CurveModel.h"
+#include "Modification/Keys/Data/KeyAttributeChangeData.h"
+
+#include <type_traits>
+
+namespace UE::CurveEditor::KeyAttributes
+{
+FKeyAttributeChangeData_PerCurve Diff(
+	const TConstArrayView<FKeyHandle> InKeysBeforeChange,
+	const TConstArrayView<FKeyAttributes> InOriginal, const TConstArrayView<FKeyAttributes> InTarget
+)
+{
+	SCOPED_CURVE_EDITOR_TRACE(Diff_KeyAttributes);
+	
+	FKeyAttributeChangeData_PerCurve ChangeData;
+	const int32 Num = InKeysBeforeChange.Num();
+	if (!ensure(InOriginal.Num() == Num && InTarget.Num() == Num))
+	{
+		return ChangeData;
+	}
+	
+	ChangeData.Handles.Reserve(Num);
+	ChangeData.BeforeChange.Reserve(Num);
+	ChangeData.AfterChange.Reserve(Num);
+	
+	for (int32 Index = 0; Index < Num; ++Index)
+	{
+		if (InOriginal[Index] != InTarget[Index])
+		{
+			ChangeData.Handles.Add(InKeysBeforeChange[Index]);
+			ChangeData.BeforeChange.Add(InOriginal[Index]);
+			ChangeData.AfterChange.Add(InTarget[Index]);
+		}
+	}
+	
+	return ChangeData;
+}
+
+namespace Private
+{
+template<typename TGetChange> requires std::is_invocable_r_v<TConstArrayView<FKeyAttributes>, TGetChange, const FKeyAttributeChangeData_PerCurve&>
+static void SetKeyAttributes(const TMap<FCurveModelID, TUniquePtr<FCurveModel>>& InCurves, const FKeyAttributeChangeData& InDeltaChange, TGetChange&& InGetChange)
+{
+	for (const TPair<FCurveModelID, FKeyAttributeChangeData_PerCurve>& Pair : InDeltaChange.ChangedCurves)
+	{
+		const TUniquePtr<FCurveModel>* FoundCurveModel = InCurves.Find(Pair.Key);
+		if (!FoundCurveModel || !FoundCurveModel->Get())
+		{
+			continue;
+		}
+		FCurveModel& CurveModel = *FoundCurveModel->Get();
+		const FKeyAttributeChangeData_PerCurve& PerCurveData = Pair.Value;
+		CurveModel.SetKeyAttributes(PerCurveData.Handles, InGetChange(PerCurveData));
+	}
+}
+}
+
+void ApplyChange(const TMap<FCurveModelID, TUniquePtr<FCurveModel>>& InCurves, const FKeyAttributeChangeData& InDeltaChange)
+{
+	Private::SetKeyAttributes(InCurves, InDeltaChange, [](const FKeyAttributeChangeData_PerCurve& Change)
+	{
+		return Change.AfterChange;
+	});}
+	
+void RevertChange(const TMap<FCurveModelID, TUniquePtr<FCurveModel>>& InCurves, const FKeyAttributeChangeData& InDeltaChange)
+{
+	Private::SetKeyAttributes(InCurves, InDeltaChange, [](const FKeyAttributeChangeData_PerCurve& Change)
+	{
+		return Change.BeforeChange;
+	});
+}
+}
diff --git a/Engine/Source/Editor/CurveEditor/Private/Modification/Keys/KeyAttributeChangeUtils.h b/Engine/Source/Editor/CurveEditor/Private/Modification/Keys/KeyAttributeChangeUtils.h
new file mode 100644
--- /dev/null
+++ b/Engine/Source/Editor/CurveEditor/Private/Modification/Keys/KeyAttributeChangeUtils.h
@@ -0,0 +1,38 @@
+﻿// Copyright Epic Games, Inc. All Rights Reserved.
+
+#pragma once
+
+#include "CurveDataAbstraction.h"
+#include "Containers/ContainersFwd.h"
+#include "Templates/UniquePtr.h"
+
+class FCurveModel;
+class IBufferedCurveModel;
+namespace UE::CurveEditor { struct FKeyAttributeChangeData; }
+namespace UE::CurveEditor { struct FKeyAttributeChangeData_PerCurve; }
+struct FCurveModelID;
+struct FKeyAttributes;
+struct FKeyHandle;
+struct FKeyPosition;
+template<typename Allocator> class TBitArray;
+
+namespace UE::CurveEditor::KeyAttributes
+{
+/**
+ * Computes the delta change to get from InOriginal to InTarget.
+ * 
+ * @param InKeysBeforeChange The keys for which to check whether they have changed.
+ * @param InOriginal The original positions of the keys.
+ * @param InTarget The new positions of the keys.
+ */
+FKeyAttributeChangeData_PerCurve Diff(
+	const TConstArrayView<FKeyHandle> InKeysBeforeChange,
+	const TConstArrayView<FKeyAttributes> InOriginal,
+	const TConstArrayView<FKeyAttributes> InTarget
+	);
+
+/** Applies InDeltaChange to InCurves (redo). */
+void ApplyChange(const TMap<FCurveModelID, TUniquePtr<FCurveModel>>& InCurves, const FKeyAttributeChangeData& InDeltaChange);
+/** Reverts InDeltaChange from InCurves (undo). */
+void RevertChange(const TMap<FCurveModelID, TUniquePtr<FCurveModel>>& InCurves, const FKeyAttributeChangeData& InDeltaChange);
+}
diff --git a/Engine/Source/Editor/CurveEditor/Private/Modification/Keys/KeyExistenceUtils.cpp b/Engine/Source/Editor/CurveEditor/Private/Modification/Keys/KeyExistenceUtils.cpp
new file mode 100644
--- /dev/null
+++ b/Engine/Source/Editor/CurveEditor/Private/Modification/Keys/KeyExistenceUtils.cpp
@@ -0,0 +1,176 @@
+﻿// Copyright Epic Games, Inc. All Rights Reserved.
+
+#include "KeyExistenceUtils.h"
+
+#include "CurveEditorTrace.h"
+#include "CurveModel.h"
+#include "Algo/Contains.h"
+#include "Modification/Keys/Data/AddKeysChangeData.h"
+#include "Modification/Keys/Data/RemoveKeysChangeData.h"
+
+namespace UE::CurveEditor
+{
+void AddKeys(FCurveModel& InCurveModel, const FCurveKeyData& InKeyChange)
+{
+	InCurveModel.PutKeys(InKeyChange.KeyHandles, InKeyChange.KeyPositions, InKeyChange.KeyAttributes);
+}
+	
+void RemoveKeys(FCurveModel& InCurveModel, const FCurveKeyData& InKeyChange, const double InCurrentSliderTime)
+{
+	InCurveModel.RemoveKeys(InKeyChange.KeyHandles, InCurrentSliderTime);
+}
+
+namespace KeyExistenceDetail
+{
+static const FCurveKeyData* FindCurveChange(const FCurveModelID& InCurveId, const FAddKeysChangeData& InChange)
+{
+	return InChange.SavedCurveState.Find(InCurveId);
+}
+static const FCurveKeyData* FindCurveChange(const FCurveModelID& InCurveId, const FRemoveKeysChangeData& InChange)
+{
+	return InChange.SavedCurveState.Find(InCurveId);
+}
+
+/** Shorthand util for applying lambda to all curve models that have FKeyDataSnapshot_PerCurve saved from them. */
+template<typename TChangeData, typename TApply> requires std::is_invocable_v<TApply, FCurveModel&, const FCurveKeyData&>
+static void ApplyKeyData(const TMap<FCurveModelID, TUniquePtr<FCurveModel>>& InCurves, const TChangeData& InChangeData, TApply&& InApplyFunctor)
+{
+	for (const TPair<FCurveModelID, TUniquePtr<FCurveModel>>& CurvePair : InCurves)
+	{
+		const TUniquePtr<FCurveModel>& CurveModel = CurvePair.Value;
+		const FCurveKeyData* PerCurveChange = FindCurveChange(CurvePair.Key, InChangeData);
+		if (ensure(CurveModel) && PerCurveChange)
+		{
+			InApplyFunctor(*CurveModel, *PerCurveChange);
+		}
+	}
+}
+
+template<typename TCallback> requires std::is_invocable_v<TCallback, int32>
+static void ForEachRemovedKeyIndex(TConstArrayView<FKeyHandle> InOriginalKeys, const FCurveModel& InCurveModel, TCallback&& InCallback)
+{
+	SCOPED_CURVE_EDITOR_TRACE(Diff_RemovedKeys);
+	
+	TArray<FKeyHandle> NewAllKeys;
+	{
+		SCOPED_CURVE_EDITOR_TRACE(Diff_RemovedKeys_GetAndSortKeys);
+		InCurveModel.GetAllKeys(NewAllKeys);
+		NewAllKeys.Sort(); // Sorting ...
+	}
+	
+	SCOPED_CURVE_EDITOR_TRACE(Diff_RemovedKeys_DoDiff);
+	TArray<FKeyHandle> Result;
+	for (int32 Index = 0; Index < InOriginalKeys.Num(); ++Index)
+	{
+		const FKeyHandle KeyHandle = InOriginalKeys[Index];
+
+		// ... and performing binary search speeds this up (profiled).
+		// In a test with 110.000 keys, we were able to get the time down from 8.1 ms to 6.5 ms.
+		const int32 FoundIndex = Algo::BinarySearch(NewAllKeys, KeyHandle);
+		const bool bIsStillContained = FoundIndex != INDEX_NONE;
+		if (!bIsStillContained)
+		{
+			InCallback(Index);
+		}
+	}
+}
+}
+
+namespace KeyInsertion
+{
+FCurveKeyData Diff(TConstArrayView<FKeyHandle> InOriginalKeys, const FCurveModel& InCurveModel)
+{
+	SCOPED_CURVE_EDITOR_TRACE(Diff_AddedKeys);
+	
+	const TArray<FKeyHandle> NewAllKeys = InCurveModel.GetAllKeys();
+	
+	FCurveKeyData Snapshot;
+	TArray<FKeyHandle>& Keys = Snapshot.KeyHandles;
+	for (int32 Index = 0; Index < NewAllKeys.Num(); ++Index)
+	{
+		const FKeyHandle KeyHandle = NewAllKeys[Index];
+		
+		// This could be sped up by using binary search maybe but not by much (for 110.000 keys, you can save 2ms).
+		// Let's not require the caller to pass it in sorted.
+		if (!InOriginalKeys.Contains(KeyHandle))
+		{
+			Keys.Add(KeyHandle);
+		}
+	}
+	
+	if (!Keys.IsEmpty())
+	{
+		Snapshot.KeyPositions.SetNum(Keys.Num());
+		Snapshot.KeyAttributes.SetNum(Keys.Num());
+		
+		InCurveModel.GetKeyPositions(Keys, Snapshot.KeyPositions);
+		InCurveModel.GetKeyAttributesExcludingAutoComputed(Keys, Snapshot.KeyAttributes);
+	}
+	return Snapshot;
+}
+
+void ApplyChange(const TMap<FCurveModelID, TUniquePtr<FCurveModel>>& InCurves, const FAddKeysChangeData& InDeltaChange)
+{
+	KeyExistenceDetail::ApplyKeyData(InCurves, InDeltaChange, [](FCurveModel& Curve, const FCurveKeyData& PerCurveData)
+	{
+		AddKeys(Curve, PerCurveData);
+	});
+}
+void RevertChange(
+	const TMap<FCurveModelID, TUniquePtr<FCurveModel>>& InCurves, const FAddKeysChangeData& InDeltaChange, const double InCurrentSliderTime
+	)
+{
+	KeyExistenceDetail::ApplyKeyData(InCurves, InDeltaChange, [InCurrentSliderTime](FCurveModel& Curve, const FCurveKeyData& PerCurveData)
+	{
+		RemoveKeys(Curve, PerCurveData, InCurrentSliderTime);
+	});
+}
+}
+
+namespace KeyRemoval
+{
+FCurveKeyData Diff(
+	TConstArrayView<FKeyHandle> InOriginalKeys, TConstArrayView<FKeyPosition> InOriginalPositions, TConstArrayView<FKeyAttributes> InOriginalAttrs,
+	const FCurveModel& InCurveModel
+	)
+{
+	FCurveKeyData Snapshot;
+	KeyExistenceDetail::ForEachRemovedKeyIndex(InOriginalKeys, InCurveModel,
+		[&InOriginalKeys, &InOriginalPositions, &InOriginalAttrs, &Snapshot](int32 Index)
+		{
+			Snapshot.KeyHandles.Add(InOriginalKeys[Index]);
+			Snapshot.KeyPositions.Add(InOriginalPositions[Index]);
+			Snapshot.KeyAttributes.Add(InOriginalAttrs[Index]);
+		});
+	return Snapshot;
+}
+
+TArray<FKeyHandle> Diff(TConstArrayView<FKeyHandle> InOriginalKeys, const FCurveModel& InCurveModel)
+{
+	TArray<FKeyHandle> Result;
+	KeyExistenceDetail::ForEachRemovedKeyIndex(InOriginalKeys, InCurveModel, [&InOriginalKeys, &Result](int32 Index)
+	{
+		Result.Add(InOriginalKeys[Index]);
+	});
+	return Result;
+}
+
+void ApplyChange(
+	const TMap<FCurveModelID, TUniquePtr<FCurveModel>>& InCurves, const FRemoveKeysChangeData& InDeltaChange, const double InCurrentSliderTime
+	)
+{
+	KeyExistenceDetail::ApplyKeyData(InCurves, InDeltaChange, [InCurrentSliderTime](FCurveModel& Curve, const FCurveKeyData& PerCurveData)
+	{
+		RemoveKeys(Curve, PerCurveData, InCurrentSliderTime);
+	});
+}
+	
+void RevertChange(const TMap<FCurveModelID, TUniquePtr<FCurveModel>>& InCurves, const FRemoveKeysChangeData& InDeltaChange)
+{
+	KeyExistenceDetail::ApplyKeyData(InCurves, InDeltaChange, [](FCurveModel& Curve, const FCurveKeyData& PerCurveData)
+	{
+		AddKeys(Curve, PerCurveData);
+	});
+}
+}
+}
diff --git a/Engine/Source/Editor/CurveEditor/Private/Modification/Keys/KeyExistenceUtils.h b/Engine/Source/Editor/CurveEditor/Private/Modification/Keys/KeyExistenceUtils.h
new file mode 100644
--- /dev/null
+++ b/Engine/Source/Editor/CurveEditor/Private/Modification/Keys/KeyExistenceUtils.h
@@ -0,0 +1,60 @@
+﻿// Copyright Epic Games, Inc. All Rights Reserved.
+
+#pragma once
+
+#include "Containers/ContainersFwd.h"
+#include "Templates/UniquePtr.h"
+
+class FCurveModel;
+namespace UE::CurveEditor { struct FCurveKeyData; }
+namespace UE::CurveEditor { struct FRemoveKeysChangeData; }
+namespace UE::CurveEditor{ struct FAddKeysChangeData;}
+struct FCurveModelID;
+struct FKeyAttributes;
+struct FKeyHandle;
+struct FKeyPosition;
+
+namespace UE::CurveEditor
+{
+/** Adds the keys saved in InKeyChange to InCurveModel. */
+void AddKeys(FCurveModel& InCurveModel, const FCurveKeyData& InKeyChange);
+/** Removes the keys saved in InKeyChange from InCurveModel. */
+void RemoveKeys(FCurveModel& InCurveModel, const FCurveKeyData& InKeyChange, const double InCurrentSliderTime = 0);
+
+namespace KeyInsertion
+{
+/**
+ * Returns the keys that are new - comparing InOriginalKeys to the current state of InCurveModel.
+ * @return The keys that have been added to InCurveModel. */
+FCurveKeyData Diff(TConstArrayView<FKeyHandle> InOriginalKeys, const FCurveModel& InCurveModel);
+
+/** Adds the keys in InDeltaChange to the curves. */
+void ApplyChange(const TMap<FCurveModelID, TUniquePtr<FCurveModel>>& InCurves, const FAddKeysChangeData& InDeltaChange);
+/** Removes the keys in InDeltaChange from the curves. */
+void RevertChange(const TMap<FCurveModelID, TUniquePtr<FCurveModel>>& InCurves, const FAddKeysChangeData& InDeltaChange,
+	const double InCurrentSliderTime = 0
+	);
+}
+
+namespace KeyRemoval
+{
+/**
+ * Returns the keys that have been removed - comparing InOriginalKeys to the current state of InCurveModel.
+ * @return The keys that have been removed from InCurveModel
+ */
+FCurveKeyData Diff(
+	TConstArrayView<FKeyHandle> InOriginalKeys, TConstArrayView<FKeyPosition> InOriginalPositions, TConstArrayView<FKeyAttributes> InOriginalAttrs,
+	const FCurveModel& InCurveModel
+	);
+
+/** @return The keys that have been removed from InCurveModel. */
+TArray<FKeyHandle> Diff(TConstArrayView<FKeyHandle> InOriginalKeys, const FCurveModel& InCurveModel);
+
+/** Removes the keys in InDeltaChange from the curves. */
+void ApplyChange(const TMap<FCurveModelID, TUniquePtr<FCurveModel>>& InCurves, const FRemoveKeysChangeData& InDeltaChange,
+	const double InCurrentSliderTime = 0
+	);
+/** Adds the keys in InDeltaChange to the curves. */
+void RevertChange(const TMap<FCurveModelID, TUniquePtr<FCurveModel>>& InCurves, const FRemoveKeysChangeData& InDeltaChange);
+}
+}
diff --git a/Engine/Source/Editor/CurveEditor/Private/Modification/Keys/MoveKeysChangeUtils.cpp b/Engine/Source/Editor/CurveEditor/Private/Modification/Keys/MoveKeysChangeUtils.cpp
new file mode 100644
--- /dev/null
+++ b/Engine/Source/Editor/CurveEditor/Private/Modification/Keys/MoveKeysChangeUtils.cpp
@@ -0,0 +1,93 @@
+﻿// Copyright Epic Games, Inc. All Rights Reserved.
+
+#include "Modification/Keys/MoveKeysChangeUtils.h"
+
+#include "CurveEditorTrace.h"
+#include "CurveModel.h"
+#include "Modification/Keys/Data/MoveKeysChangeData.h"
+
+namespace UE::CurveEditor::MoveKeys
+{
+namespace Private
+{
+template<typename TTranslateOp> requires std::is_invocable_r_v<FKeyPosition, TTranslateOp, const FKeyPosition&, const FVector2f&>
+static void CombineData(
+	const TMap<FCurveModelID, TUniquePtr<FCurveModel>>& InCurves,
+	const FMoveKeysChangeData& InDeltaChange,
+	TTranslateOp&& InTranslateFunc
+	)
+{
+	for (const TPair<FCurveModelID, FMoveKeysChangeData_PerCurve>& Pair : InDeltaChange.ChangedCurves)
+	{
+		const TUniquePtr<FCurveModel>* FoundCurveModel = InCurves.Find(Pair.Key);
+		if (!FoundCurveModel || !FoundCurveModel->Get())
+		{
+			continue;
+		}
+		FCurveModel& CurveModel = *FoundCurveModel->Get();
+	
+		const FMoveKeysChangeData_PerCurve& Data = Pair.Value;
+		const int32 Num = Data.Handles.Num();
+		check(Data.Translations.Num() == Num);
+	
+		TArray<FKeyPosition> CurrentPositions;
+		CurrentPositions.SetNumUninitialized(Num);
+		CurveModel.GetKeyPositions(Data.Handles, CurrentPositions);
+	
+		for (int32 Index = 0; Index < Data.Handles.Num(); ++Index)
+		{
+			CurrentPositions[Index] = InTranslateFunc(CurrentPositions[Index], Data.Translations[Index]);
+		}
+
+		CurveModel.SetKeyPositions(Data.Handles, CurrentPositions, EPropertyChangeType::ValueSet);
+	}
+}
+}
+FMoveKeysChangeData_PerCurve Diff(
+	const TConstArrayView<FKeyHandle> InKeys,
+	const TConstArrayView<FKeyPosition> InOriginal, const TConstArrayView<FKeyPosition> InTarget
+	)
+{
+	SCOPED_CURVE_EDITOR_TRACE(Diff_MovedKeys);
+	
+	const int32 Num = InKeys.Num();
+	if (!ensure(InOriginal.Num() == Num && InTarget.Num() == Num))
+	{
+		return {};
+	}
+
+	FMoveKeysChangeData_PerCurve Result;
+	Result.Handles.Reserve(Num);
+	Result.Translations.Reserve(Num);
+	for (int32 Index = 0; Index < Num; ++Index)
+	{
+		const FKeyPosition& Original = InOriginal[Index];
+		const FKeyPosition& Target = InTarget[Index];
+		const FVector2f Delta = FVector2f(Target.InputValue - Original.InputValue, Target.OutputValue - Original.OutputValue);
+
+		if (!Delta.IsNearlyZero())
+		{
+			Result.Handles.Add(InKeys[Index]);
+			Result.Translations.Add(Delta);
+		}
+	}
+	
+	return Result;
+}
+
+void ApplyChange(const TMap<FCurveModelID, TUniquePtr<FCurveModel>>& InCurves, const FMoveKeysChangeData& InDeltaChange)
+{
+	Private::CombineData(InCurves, InDeltaChange, [](const FKeyPosition& Position, const FVector2f Delta)
+	{
+		return FKeyPosition(Position.InputValue + Delta.X, Position.OutputValue + Delta.Y);
+	});
+}
+
+void RevertChange(const TMap<FCurveModelID, TUniquePtr<FCurveModel>>& InCurves, const FMoveKeysChangeData& InDeltaChange)
+{
+	Private::CombineData(InCurves, InDeltaChange, [](const FKeyPosition& Position, const FVector2f Delta)
+	{
+		return FKeyPosition(Position.InputValue - Delta.X, Position.OutputValue - Delta.Y);
+	});
+}
+}
diff --git a/Engine/Source/Editor/CurveEditor/Private/Modification/Keys/MoveKeysChangeUtils.h b/Engine/Source/Editor/CurveEditor/Private/Modification/Keys/MoveKeysChangeUtils.h
new file mode 100644
--- /dev/null
+++ b/Engine/Source/Editor/CurveEditor/Private/Modification/Keys/MoveKeysChangeUtils.h
@@ -0,0 +1,38 @@
+﻿// Copyright Epic Games, Inc. All Rights Reserved.
+
+#pragma once
+
+#include "Containers/BitArray.h"
+#include "Containers/ContainersFwd.h"
+#include "Templates/UniquePtr.h"
+
+class FCurveModel;
+class IBufferedCurveModel;
+namespace UE::CurveEditor { struct FMoveKeysChangeData; }
+namespace UE::CurveEditor { struct FMoveKeysChangeData_PerCurve; }
+struct FCurveModelID;
+struct FKeyHandle;
+struct FKeyPosition;
+template<typename Allocator> class TBitArray;
+
+namespace UE::CurveEditor::MoveKeys
+{
+/**
+ * Computes the delta change to get from InOriginal to InTarget.
+ * 
+ * @param InKeysBeforeChange The keys for which to check whether they have moved.
+ * @param InOriginal The original positions of the keys.
+ * @param InTarget The new positions of the keys.
+ */
+FMoveKeysChangeData_PerCurve Diff(
+	const TConstArrayView<FKeyHandle> InKeysBeforeChange,
+	const TConstArrayView<FKeyPosition> InOriginal,
+	const TConstArrayView<FKeyPosition> InTarget
+	);
+
+/** Applies InDeltaChange to InCurves (redo). */
+void ApplyChange(const TMap<FCurveModelID, TUniquePtr<FCurveModel>>& InCurves, const FMoveKeysChangeData& InDeltaChange);
+
+/** Reverts InDeltaChange from InCurves (undo). */
+void RevertChange(const TMap<FCurveModelID, TUniquePtr<FCurveModel>>& InCurves, const FMoveKeysChangeData& InDeltaChange);
+}
diff --git a/Engine/Source/Editor/CurveEditor/Private/Modification/Utils/ScopedKeyChange.cpp b/Engine/Source/Editor/CurveEditor/Private/Modification/Utils/ScopedKeyChange.cpp
new file mode 100644
--- /dev/null
+++ b/Engine/Source/Editor/CurveEditor/Private/Modification/Utils/ScopedKeyChange.cpp
@@ -0,0 +1,88 @@
+﻿// Copyright Epic Games, Inc. All Rights Reserved.
+
+#include "Modification/Utils/ScopedKeyChange.h"
+
+#include "CurveEditor.h"
+#include "CurveModel.h"
+#include "HAL/IConsoleManager.h"
+#include "Internationalization/Text.h"
+#include "Misc/ITransaction.h"
+#include "Misc/ScopeExit.h"
+#include "Modification/Keys/Data/GenericCurveChangeData.h"
+#include "Modification/Keys/GenericCurveChangeCommand.h"
+#include "Modification/Keys/GenericCurveChangeUtils.h"
+#include "Modification/Keys/SnapshotPerfDebugUtils.h"
+#include "ProfilingDebugging/ScopedTimers.h"
+#include "SCurveEditor.h"
+
+namespace UE::CurveEditor
+{
+static TAutoConsoleVariable<bool> CVarLogUndoKeyChangeSize(
+	TEXT("CurveEditor.LogDiffPerfStats"),
+	false,
+	TEXT("Whenever a key change is made, whether to log debug info (time spent diffing, bytes used for the diff, the final bytes on undo stack size).")
+	);
+
+FScopedKeyChange::FScopedKeyChange(
+	TWeakPtr<FCurveEditor> InCurveEditor, FCurveChangeDiff InKeyChangeDiff, const FText& InDescription, EScopedKeyChangeFlags InFlags
+	)
+	: FScopedChangeBase(MoveTemp(InCurveEditor), InDescription.IsEmpty() ? NSLOCTEXT("CurveEditor", "ChangeKeys", "Change Keys") : InDescription)
+	, ChangeDiff(MoveTemp(InKeyChangeDiff))
+	, Flags(InFlags)
+{}
+
+FScopedKeyChange::FScopedKeyChange(
+	FCurvesSnapshotBuilder InDataToDiff, const FText& InDescription, EScopedKeyChangeFlags InFlags
+	)
+	: FScopedKeyChange(
+		InDataToDiff.WeakCurveEditor,
+		FCurveChangeDiff(InDataToDiff.WeakCurveEditor, MoveTemp(InDataToDiff.Snapshot), InDataToDiff.CapturedDataFlags),
+		InDescription, InFlags
+		)
+{}
+
+FScopedKeyChange::~FScopedKeyChange()
+{
+	const TSharedPtr<FCurveEditor> CurveEditorPin = WeakCurveEditor.Pin();
+	if (!CurveEditorPin)
+	{
+		return;
+	}
+
+	// Early out without expensive call to ComputeDiff.
+	const bool bRevertOnCancelled = EnumHasAnyFlags(Flags, EScopedKeyChangeFlags::RevertOnCancel);
+	const bool bHasBeenCancelled = IsCancelled();
+	if (bHasBeenCancelled && !bRevertOnCancelled)
+	{
+		return;
+	}
+
+	FAutoScopedDurationTimer TimeMeasurement;
+	FGenericCurveChangeData DeltaChange = ChangeDiff.ComputeDiff();
+	const auto LogPerfData = [this, &TimeMeasurement, &DeltaChange]
+	{
+		UE_CLOG(CVarLogUndoKeyChangeSize.GetValueOnAnyThread(), LogCurveEditor, Log, TEXT("Change'%s' --- Undo size: %s ---\t Diff: Time %f ms, %s"),
+			GUndo ? *GUndo->GetContext().Title.ToString() : *Description.ToString(),
+			*FText::AsMemory(DeltaChange.GetAllocatedSize()).ToString(),
+			FTimespan::FromSeconds(TimeMeasurement.GetTime()).GetTotalMilliseconds(),
+			*DumpSnapshotPerfData(ChangeDiff.GetSnapshot(), ECurvesSnapshotPerfFlags::All)
+			); 
+	};
+	
+	if (!DeltaChange)
+	{
+		return;
+	}
+	
+	if (!bHasBeenCancelled)
+	{
+		LogPerfData(); // Important to do this before MoveTemp
+		TryAppendCommand(MakeUnique<FGenericCurveChangeCommand>(MoveTemp(DeltaChange)));
+	}
+	else
+	{
+		LogPerfData();
+		GenericCurveChange::RevertChange(*CurveEditorPin, DeltaChange);
+	}
+}
+}
diff --git a/Engine/Source/Editor/CurveEditor/Private/Modification/Utils/ScopedSelectionChange.cpp b/Engine/Source/Editor/CurveEditor/Private/Modification/Utils/ScopedSelectionChange.cpp
--- a/Engine/Source/Editor/CurveEditor/Private/Modification/Utils/ScopedSelectionChange.cpp
+++ b/Engine/Source/Editor/CurveEditor/Private/Modification/Utils/ScopedSelectionChange.cpp
@@ -11,12 +11,13 @@
 
 namespace UE::CurveEditor
 {
-FScopedSelectionChange::FScopedSelectionChange(TWeakPtr<FCurveEditor> InCurveEditor)
-	: FScopedSelectionChange(MoveTemp(InCurveEditor), LOCTEXT("ChangeSelection", "Change selection"))
+FScopedSelectionChange::FScopedSelectionChange(TWeakPtr<FCurveEditor> InCurveEditor, EScopedSelectionChangeFlags InFlags)
+	: FScopedSelectionChange(MoveTemp(InCurveEditor), LOCTEXT("ChangeSelection", "Change selection"), InFlags)
 {}
 
-FScopedSelectionChange::FScopedSelectionChange(TWeakPtr<FCurveEditor> InCurveEditor, const FText& InDescription)
+FScopedSelectionChange::FScopedSelectionChange(TWeakPtr<FCurveEditor> InCurveEditor, const FText& InDescription, EScopedSelectionChangeFlags InFlags)
 	: FScopedChangeBase(MoveTemp(InCurveEditor), InDescription)
+	, Flags(InFlags)
 {
 	if (const TSharedPtr<FCurveEditor> CurveEditorPin = WeakCurveEditor.Pin())
 	{
@@ -31,11 +32,24 @@ FScopedSelectionChange::~FScopedSelectionChange()
 	{
 		return;
 	}
-	
-	if (FSelectionDeltaChange DeltaChange = DiffSelection(OriginalSelection, CurveEditorPin->Selection))
+
+	// Early out without expensive call to DiffSelection.
+	const bool bRevertOnCancelled = EnumHasAnyFlags(Flags, EScopedSelectionChangeFlags::RevertOnCancel);
+	const bool bHasBeenCancelled = IsCancelled();
+	if (bHasBeenCancelled && !bRevertOnCancelled)
+	{
+		return;
+	}
+
+	FSelectionDeltaChange DeltaChange = DiffSelection(OriginalSelection, CurveEditorPin->Selection);
+	if (!bHasBeenCancelled)
 	{
 		TryAppendCommand(MakeUnique<FSelectionChangeCommand>(MoveTemp(DeltaChange)));
 	}
+	else
+	{
+		RevertSelectionChange(CurveEditorPin->Selection, DeltaChange);
+	}
 }
 }
 
diff --git a/Engine/Source/Editor/CurveEditor/Public/Modification/Keys/CurveChangeDiff.h b/Engine/Source/Editor/CurveEditor/Public/Modification/Keys/CurveChangeDiff.h
new file mode 100644
--- /dev/null
+++ b/Engine/Source/Editor/CurveEditor/Public/Modification/Keys/CurveChangeDiff.h
@@ -0,0 +1,69 @@
+﻿// Copyright Epic Games, Inc. All Rights Reserved.
+
+#pragma once
+
+#include "Containers/Map.h"
+#include "CurvesSnapshotBuilder.h"
+#include "Data/CurveKeyDataMap.h"
+#include "Data/GenericCurveChangeData.h"
+#include "ECurveChangeFlags.h"
+#include "Templates/SharedPointer.h"
+
+class FCurveEditor;
+
+namespace UE::CurveEditor
+{
+struct FGenericCurveChangeData;
+
+/** Computes what has changed for the data saved in InSnapshot. */
+CURVEEDITOR_API FGenericCurveChangeData ComputeDiff(
+	const FCurveEditor& InCurveEditor, const ECurveChangeFlags InOperations, const FCurvesSnapshot& InSnapshot
+	);
+
+/**
+ * Utility that holds a snapshot of some curve data for the purposes of later comparing changes made to those curves.
+ * 
+ * This is useful for actions that take multiple frames, e.g. the user dragging keys.
+ * Start by constructing FCurvesSnapshotBuilder instance, adding the curves you want to diff later, and when the operation is done call ComputeDiff to
+ * compute the change.
+ *
+ * Note: This class introduces some overhead by temporarily saving the curve state.
+ * If the change is a one-off (i.e. not interactive) you usually know exactly what will change, so you can consider building the change directly.
+ * It's still fine to use this class in that case as it simplifies the code: the overhead is usually acceptable.
+ * @see FScopedKeyChange For some profiling data.
+ */
+class FCurveChangeDiff
+{
+public:
+
+	explicit FCurveChangeDiff(TWeakPtr<FCurveEditor> InCurveEditor, FCurvesSnapshot InSnapshot, ECurveChangeFlags InOps = ECurveChangeFlags::All)
+		: WeakCurveEditor(MoveTemp(InCurveEditor))
+		, Operations(InOps)
+		, Snapshot(MoveTemp(InSnapshot))
+	{}
+	
+	FCurveChangeDiff(FCurvesSnapshotBuilder InBuilder)
+		: FCurveChangeDiff(InBuilder.WeakCurveEditor, MoveTemp(InBuilder.Snapshot), InBuilder.CapturedDataFlags)
+	{}
+
+	/** Computes what has changed on the given curves. */
+	FGenericCurveChangeData ComputeDiff() const
+	{
+		const TSharedPtr<FCurveEditor> CurveEditorPin = WeakCurveEditor.Pin();
+		return CurveEditorPin ? CurveEditor::ComputeDiff(*CurveEditorPin, Operations, Snapshot) : FGenericCurveChangeData();
+	}
+
+	/** The snapshot against which is diffed. */
+	const FCurvesSnapshot& GetSnapshot() const { return Snapshot; }
+
+private:
+
+	/** Needed to look up curves and compare state. */
+	const TWeakPtr<FCurveEditor> WeakCurveEditor;
+	/** The operations that we should diff with. */
+	const ECurveChangeFlags Operations;
+
+	/** Data that ComputeDiff diffs against. */
+	const FCurvesSnapshot Snapshot;
+};
+}
diff --git a/Engine/Source/Editor/CurveEditor/Public/Modification/Keys/CurvesSnapshotBuilder.h b/Engine/Source/Editor/CurveEditor/Public/Modification/Keys/CurvesSnapshotBuilder.h
new file mode 100644
--- /dev/null
+++ b/Engine/Source/Editor/CurveEditor/Public/Modification/Keys/CurvesSnapshotBuilder.h
@@ -0,0 +1,80 @@
+﻿// Copyright Epic Games, Inc. All Rights Reserved.
+
+#pragma once
+
+#include "CurveDataAbstraction.h"
+#include "Data/CurvesSnapshot.h"
+#include "ECurveChangeFlags.h"
+
+class FCurveEditor;
+
+namespace UE::CurveEditor
+{
+/**
+ * Utility for selectively capturing state of curves based on some ECurveChangeFlags.
+ * Intended for constructing either a FCurvesSnapshot or FCurveChangeDiff.
+ */
+struct FCurvesSnapshotBuilder
+{
+	/** Needed to look up curves and compare state. */
+	const TWeakPtr<FCurveEditor> WeakCurveEditor;
+	/** The operations that we should diff with. */
+	const ECurveChangeFlags CapturedDataFlags;
+
+	/** Data that ComputeDiff diffs against. */
+	FCurvesSnapshot Snapshot;
+
+	explicit FCurvesSnapshotBuilder(
+		TWeakPtr<FCurveEditor> InCurveEditor, ECurveChangeFlags InOps = ECurveChangeFlags::All
+		)
+		: WeakCurveEditor(MoveTemp(InCurveEditor))
+		, CapturedDataFlags(InOps)
+	{}
+
+	/** Util constructor for tracking a single curve. */
+	explicit FCurvesSnapshotBuilder(
+		TWeakPtr<FCurveEditor> InCurveEditor, const FCurveModelID& InCurveId, ECurveChangeFlags InOps = ECurveChangeFlags::All
+		)
+		: FCurvesSnapshotBuilder(MoveTemp(InCurveEditor), InOps)
+	{
+		TrackCurve(InCurveId);
+	}
+	/** Util constructor for tracking a container of curves (TSet, TMap, TArray, TArrayView, TConstArrayView) */
+	template<typename TRangeType>
+	explicit FCurvesSnapshotBuilder(
+		TWeakPtr<FCurveEditor> InCurveEditor, const TRangeType& InRangeType, ECurveChangeFlags InOps = ECurveChangeFlags::All
+		)
+		: FCurvesSnapshotBuilder(MoveTemp(InCurveEditor), InOps)
+	{
+		TrackCurves(InRangeType);
+	}
+
+	/** Tracks all keys on InCurveModel. */
+	CURVEEDITOR_API FCurvesSnapshotBuilder& TrackCurve(const FCurveModelID& InCurveId);
+	CURVEEDITOR_API FCurvesSnapshotBuilder& TrackCurves(TConstArrayView<FCurveModelID> InCurves);
+	CURVEEDITOR_API FCurvesSnapshotBuilder& TrackCurves(const TSet<FCurveModelID>& InCurves);
+	template<typename T>
+	FCurvesSnapshotBuilder& TrackCurves(const TMap<FCurveModelID, T>& InCurveMapping);
+	
+	/** Tracks only the specified keys. Incompatible with EKeyChangeOperationFlags::AddKeys. */
+	CURVEEDITOR_API FCurvesSnapshotBuilder& TrackCurve(const FCurveModelID& InCurveId, TConstArrayView<FKeyHandle> Keys);
+	
+	/** Tracks all curves that are selected: adds only their selected keys. Incompatible with EKeyChangeOperationFlags::AddKeys. */
+	CURVEEDITOR_API FCurvesSnapshotBuilder& TrackSelectedKeys();
+	/** Tracks all curves that are selected: adds all keys of those curves. */
+	CURVEEDITOR_API FCurvesSnapshotBuilder& TrackSelectedCurves();
+
+	/** Tracks all curves. */
+	CURVEEDITOR_API FCurvesSnapshotBuilder& TrackAllCurves();
+};
+
+template <typename T>
+FCurvesSnapshotBuilder& FCurvesSnapshotBuilder::TrackCurves(const TMap<FCurveModelID, T>& InCurveMapping)
+{
+	for (const TPair<FCurveModelID, T>& Pair : InCurveMapping)
+	{
+		TrackCurve(Pair.Key);
+	}
+	return *this;
+}
+}
diff --git a/Engine/Source/Editor/CurveEditor/Public/Modification/Keys/Data/AddKeysChangeData.h b/Engine/Source/Editor/CurveEditor/Public/Modification/Keys/Data/AddKeysChangeData.h
new file mode 100644
--- /dev/null
+++ b/Engine/Source/Editor/CurveEditor/Public/Modification/Keys/Data/AddKeysChangeData.h
@@ -0,0 +1,21 @@
+﻿// Copyright Epic Games, Inc. All Rights Reserved.
+
+#pragma once
+
+#include "CurveKeyDataMap.h"
+
+namespace UE::CurveEditor
+{
+struct FAddKeysChangeData
+{
+	/**
+	 * On undo, remove the FKeyDataSnapshot::Keys.
+	 * On redo, add back FKeyDataSnapshot::Keys together with the saved positions and attributes.
+	 */
+	TMap<FCurveModelID, FCurveKeyData> SavedCurveState;
+	
+	bool HasChanges() const { return !SavedCurveState.IsEmpty(); }
+	operator bool() const { return HasChanges(); }
+	SIZE_T GetAllocatedSize() const { return FCurveKeyData::GetAllocatedSize(SavedCurveState); }
+};
+}
diff --git a/Engine/Source/Editor/CurveEditor/Public/Modification/Keys/Data/CurveAttributeChangeData.h b/Engine/Source/Editor/CurveEditor/Public/Modification/Keys/Data/CurveAttributeChangeData.h
new file mode 100644
--- /dev/null
+++ b/Engine/Source/Editor/CurveEditor/Public/Modification/Keys/Data/CurveAttributeChangeData.h
@@ -0,0 +1,33 @@
+﻿// Copyright Epic Games, Inc. All Rights Reserved.
+
+#pragma once
+#include "CurveEditorTypes.h"
+#include "Curves/RealCurve.h"
+
+namespace UE::CurveEditor
+{
+/** Stores data for restoring the curve attributes of a single curve. */
+struct FCurveAttributeChangeData_PerCurve
+{
+	/** The key attributes before the change. Set attributes to this on undo.  */
+	FCurveAttributes BeforeChange;
+	/** The key attributes after the change. Set attributes to this on undo. */
+	FCurveAttributes AfterChange;
+
+	bool HasChanges() const { return BeforeChange != AfterChange; }
+	operator bool() const { return HasChanges(); }
+
+	friend bool operator==(const FCurveAttributeChangeData_PerCurve&, const FCurveAttributeChangeData_PerCurve&) = default;
+	friend bool operator!=(const FCurveAttributeChangeData_PerCurve&, const FCurveAttributeChangeData_PerCurve&) = default;
+};
+
+/** Stores data for restoring curves attributes of multiple curves. */
+struct FCurveAttributeChangeData
+{
+	TMap<FCurveModelID, FCurveAttributeChangeData_PerCurve> ChangeData;
+
+	bool HasChanges() const { return !ChangeData.IsEmpty(); }
+	operator bool() const { return HasChanges(); }
+	SIZE_T GetAllocatedSize() const { return ChangeData.GetAllocatedSize(); }
+};
+}
diff --git a/Engine/Source/Editor/CurveEditor/Public/Modification/Keys/Data/CurveKeyDataMap.h b/Engine/Source/Editor/CurveEditor/Public/Modification/Keys/Data/CurveKeyDataMap.h
new file mode 100644
--- /dev/null
+++ b/Engine/Source/Editor/CurveEditor/Public/Modification/Keys/Data/CurveKeyDataMap.h
@@ -0,0 +1,66 @@
+﻿// Copyright Epic Games, Inc. All Rights Reserved.
+
+#pragma once
+
+#include "Containers/Array.h"
+#include "CurveDataAbstraction.h"
+#include "Algo/Accumulate.h"
+#include "Misc/Optional.h"
+
+namespace UE::CurveEditor
+{
+/** Holds data for restoring keys on a curve. */
+struct FCurveKeyData
+{
+	/**
+	 * The tracked key handles. Contains all key handles if EKeyChangeOperationFlags::AddKeys.
+	 * Always set.
+	 */
+	TArray<FKeyHandle> KeyHandles;
+
+	/**
+	 * The positions of Handles. Each index matches that of Handles.
+	 * Set only for EKeyChangeOperationFlags::MoveKeys.
+	 *
+	 * Either same length as Handles or 0 length.
+	 */
+	TArray<FKeyPosition> KeyPositions;
+
+	/**
+	 * The attributes of Handles. Each index matches that of Handles.
+	 * Set only for EKeyChangeOperationFlags::KeyAttributes.
+	 *
+	 * Either same length as Handles or 0 length.
+	 */
+	TArray<FKeyAttributes> KeyAttributes;
+
+	bool HasKeyPositions() const { return !KeyPositions.IsEmpty(); }
+	bool HasKeyAttributes() const { return !KeyAttributes.IsEmpty(); }
+	
+	bool HasData() const { return !KeyHandles.IsEmpty(); }
+	operator bool() const { return HasData(); }
+	
+	SIZE_T GetAllocatedSize() const { return KeyHandles.GetAllocatedSize() + KeyPositions.GetAllocatedSize() + KeyAttributes.GetAllocatedSize(); }
+	
+	template<typename T>
+	static SIZE_T GetAllocatedSize(const TMap<T, FCurveKeyData>& InMapping)
+	{
+		return Algo::TransformAccumulate(InMapping, [](const TPair<FCurveModelID, FCurveKeyData>& Pair)
+		{
+			return Pair.Value.GetAllocatedSize();
+		}, 0);
+	}
+};
+
+/** Snapshot of curve editor content. */
+struct FCurveKeyDataMap
+{
+	/** The state that the curves had. */
+	TMap<FCurveModelID, FCurveKeyData> SavedCurveState;
+
+	FCurveKeyDataMap() = default;
+	explicit FCurveKeyDataMap(TMap<FCurveModelID, FCurveKeyData> SavedCurveState) : SavedCurveState(MoveTemp(SavedCurveState)) {}
+
+	SIZE_T GetAllocatedSize() const { return FCurveKeyData::GetAllocatedSize(SavedCurveState); }
+};
+}
diff --git a/Engine/Source/Editor/CurveEditor/Public/Modification/Keys/Data/CurvesSnapshot.h b/Engine/Source/Editor/CurveEditor/Public/Modification/Keys/Data/CurvesSnapshot.h
new file mode 100644
--- /dev/null
+++ b/Engine/Source/Editor/CurveEditor/Public/Modification/Keys/Data/CurvesSnapshot.h
@@ -0,0 +1,58 @@
+﻿// Copyright Epic Games, Inc. All Rights Reserved.
+
+#pragma once
+
+#include "Algo/Accumulate.h"
+#include "CurveKeyDataMap.h"
+#include "Curves/RealCurve.h"
+#include "Misc/Optional.h"
+
+namespace UE::CurveEditor
+{
+/**
+ * Holds data that is captured to diff a curves previous state to a current state.
+ * 
+ * Not all data needs to be captured, e.g. Handles, KeyPositions, Attributes can be empty and CurveAttributes unset.
+ * It would be pointless to construct this without any members set.
+ * If either KeyPositions or Attributes is set, then Handles is also set.
+ * Intended to be built using FCurvesSnapshotBuilder and ECurveChangeFlags.
+ */
+struct FCurveDiffingData : FCurveKeyData
+{
+	/** The curve attributes if they were captured. */
+	TOptional<FCurveAttributes> CurveAttributes;
+
+	bool HasCurveAttributes() const { return CurveAttributes.IsSet(); }
+	
+	bool HasData() const { return FCurveKeyData::HasData() || CurveAttributes.IsSet(); }
+	operator bool() const { return FCurveDiffingData::HasData(); }
+	
+	SIZE_T GetAllocatedSize() const { return FCurveKeyData::GetAllocatedSize() + sizeof(CurveAttributes); }
+	
+	template<typename T>
+	static SIZE_T GetAllocatedSize(const TMap<T, FCurveDiffingData>& InMapping)
+	{
+		return Algo::TransformAccumulate(InMapping, [](const TPair<FCurveModelID, FCurveDiffingData>& Pair)
+		{
+			return Pair.Value.GetAllocatedSize();
+		}, 0);
+	}
+};
+
+/** Holds selected data about curves. This data is usually used to diff against changes made to those curves. */
+struct FCurvesSnapshot
+{
+	/**
+	 * The info about the curves.
+	 *
+	 * Not all members of FCurveDiffingData need to contain data: just the members intended for diffing.
+	 * Intended to be built using FCurvesSnapshotBuilder and ECurveChangeFlags.
+	 */
+	TMap<FCurveModelID, FCurveDiffingData> CurveData;
+
+	FCurvesSnapshot() = default;
+	explicit FCurvesSnapshot(TMap<FCurveModelID, FCurveDiffingData> InCurveData) : CurveData(MoveTemp(InCurveData)) {}
+	
+	SIZE_T GetAllocatedSize() const { return FCurveDiffingData::GetAllocatedSize(CurveData); }
+};
+}
diff --git a/Engine/Source/Editor/CurveEditor/Public/Modification/Keys/Data/GenericCurveChangeData.h b/Engine/Source/Editor/CurveEditor/Public/Modification/Keys/Data/GenericCurveChangeData.h
new file mode 100644
--- /dev/null
+++ b/Engine/Source/Editor/CurveEditor/Public/Modification/Keys/Data/GenericCurveChangeData.h
@@ -0,0 +1,43 @@
+﻿// Copyright Epic Games, Inc. All Rights Reserved.
+
+#pragma once
+
+#include "AddKeysChangeData.h"
+#include "CurveAttributeChangeData.h"
+#include "KeyAttributeChangeData.h"
+#include "MoveKeysChangeData.h"
+#include "RemoveKeysChangeData.h"
+
+namespace UE::CurveEditor
+{
+/** Stores changes made to curves in the curve editor. */
+struct FGenericCurveChangeData
+{
+	FMoveKeysChangeData MoveKeysData;
+	FAddKeysChangeData AddKeysData;
+	FRemoveKeysChangeData RemoveKeysData;
+	FKeyAttributeChangeData KeyAttributeData;
+	FCurveAttributeChangeData CurveAttributeData;
+	
+	bool HasChanges() const { return MoveKeysData.HasChanges() || AddKeysData.HasChanges() || RemoveKeysData.HasChanges() || KeyAttributeData.HasChanges() || CurveAttributeData.HasChanges(); }
+	operator bool() const { return HasChanges(); }
+
+	/** @return The number of distinct change types. Useful for testing. */
+	int32 NumChangeTypes() const
+	{
+		return static_cast<int32>(MoveKeysData.HasChanges())
+		+ static_cast<int32>(AddKeysData.HasChanges())
+		+ static_cast<int32>(RemoveKeysData.HasChanges())
+		+ static_cast<int32>(KeyAttributeData.HasChanges())
+		+ static_cast<int32>(CurveAttributeData.HasChanges());
+	}
+	SIZE_T GetAllocatedSize() const
+	{
+		return MoveKeysData.GetAllocatedSize()
+			+ AddKeysData.GetAllocatedSize()
+			+ RemoveKeysData.GetAllocatedSize()
+			+ KeyAttributeData.GetAllocatedSize()
+			+ CurveAttributeData.GetAllocatedSize();
+	}
+};
+}
diff --git a/Engine/Source/Editor/CurveEditor/Public/Modification/Keys/Data/KeyAttributeChangeData.h b/Engine/Source/Editor/CurveEditor/Public/Modification/Keys/Data/KeyAttributeChangeData.h
new file mode 100644
--- /dev/null
+++ b/Engine/Source/Editor/CurveEditor/Public/Modification/Keys/Data/KeyAttributeChangeData.h
@@ -0,0 +1,50 @@
+﻿// Copyright Epic Games, Inc. All Rights Reserved.
+
+#pragma once
+
+#include "Algo/Accumulate.h"
+#include "Containers/Array.h"
+#include "Containers/Map.h"
+#include "CurveDataAbstraction.h"
+#include "CurveEditorTypes.h"
+
+namespace UE::CurveEditor
+{
+/** Stores data for restoring key attributes for a single curve. */
+struct FKeyAttributeChangeData_PerCurve
+{
+	/** The key handles that changed attributes. */
+	TArray<FKeyHandle> Handles;
+
+	// Ideally, we'd only save the delta change values.
+	// Using 2 arrays like this, uses 2x amount of RAM than the optimal solution would. 
+	// However, the tradeoff is that this implementation is easy to understand (less error-prone).
+	// I tried the "ideal" solution: there's many more details to take care of. So I decided, it's not worth to engineer anything better for now.
+	
+	/** The key attributes before the change. Set attributes to this on undo. Each index corresponds to Handles. */
+	TArray<FKeyAttributes> BeforeChange;
+	/** The key attributes after the change. Set attributes to this on redo. Each index corresponds to Handles. */
+	TArray<FKeyAttributes> AfterChange;
+
+	bool HasChanges() const { return !Handles.IsEmpty(); }
+	operator bool() const { return HasChanges(); }
+	SIZE_T GetAllocatedSize() const { return Handles.GetAllocatedSize() + BeforeChange.GetAllocatedSize() + AfterChange.GetAllocatedSize(); }
+};
+
+/** Stores data for restoring key attributes of multiple curves. */
+struct FKeyAttributeChangeData
+{
+	TMap<FCurveModelID, FKeyAttributeChangeData_PerCurve> ChangedCurves;
+	
+	bool HasChanges() const { return !ChangedCurves.IsEmpty(); }
+	operator bool() const { return HasChanges(); }
+	SIZE_T GetAllocatedSize() const
+	{
+		return ChangedCurves.GetAllocatedSize()
+			+ Algo::TransformAccumulate(ChangedCurves, [](const TPair<FCurveModelID, FKeyAttributeChangeData_PerCurve>& Pair)
+			{
+				return Pair.Value.GetAllocatedSize();
+			}, 0);
+	}
+};
+}
diff --git a/Engine/Source/Editor/CurveEditor/Public/Modification/Keys/Data/MoveKeysChangeData.h b/Engine/Source/Editor/CurveEditor/Public/Modification/Keys/Data/MoveKeysChangeData.h
new file mode 100644
--- /dev/null
+++ b/Engine/Source/Editor/CurveEditor/Public/Modification/Keys/Data/MoveKeysChangeData.h
@@ -0,0 +1,50 @@
+﻿// Copyright Epic Games, Inc. All Rights Reserved.
+
+#pragma once
+
+#include "Algo/Accumulate.h"
+#include "Containers/Array.h"
+#include "Containers/Map.h"
+#include "CurveEditorTypes.h"
+
+namespace UE::CurveEditor
+{
+struct FMoveKeysChangeData_PerCurve
+{
+	/** The moved key handles. */
+	TArray<FKeyHandle> Handles;
+	
+	/**
+	 * The amount the corresponding key was moved.
+	 * 
+	 * When applying (redo) the change, add the value to the key position.
+	 * When reverting (undo) the change, subtract the value.
+	 */
+	TArray<FVector2f> Translations;
+
+	bool HasChanges() const { return !Handles.IsEmpty(); }
+	operator bool() const { return HasChanges(); }
+	SIZE_T GetAllocatedSize() const { return Handles.GetAllocatedSize() + Translations.GetAllocatedSize(); }
+};
+
+/**
+ * Data for an operation that translates keys.
+ * This can still be improved: if all keys were transformed uniformly, then we only need to store a single FVector2f.
+ */
+struct FMoveKeysChangeData
+{
+	/** The per curve changes. */
+	TMap<FCurveModelID, FMoveKeysChangeData_PerCurve> ChangedCurves;
+	
+	bool HasChanges() const { return !ChangedCurves.IsEmpty(); }
+	operator bool() const { return HasChanges(); }
+	SIZE_T GetAllocatedSize() const
+	{
+		return ChangedCurves.GetAllocatedSize()
+			+ Algo::TransformAccumulate(ChangedCurves, [](const TPair<FCurveModelID, FMoveKeysChangeData_PerCurve>& Pair)
+			{
+				return Pair.Value.GetAllocatedSize();
+			}, 0);
+	}
+};
+}
diff --git a/Engine/Source/Editor/CurveEditor/Public/Modification/Keys/Data/RemoveKeysChangeData.h b/Engine/Source/Editor/CurveEditor/Public/Modification/Keys/Data/RemoveKeysChangeData.h
new file mode 100644
--- /dev/null
+++ b/Engine/Source/Editor/CurveEditor/Public/Modification/Keys/Data/RemoveKeysChangeData.h
@@ -0,0 +1,23 @@
+﻿// Copyright Epic Games, Inc. All Rights Reserved.
+
+#pragma once
+
+#include "Containers/Map.h"
+#include "CurveEditorTypes.h"
+#include "CurveKeyDataMap.h"
+
+namespace UE::CurveEditor
+{
+struct FRemoveKeysChangeData
+{
+	/**
+	 * On undo, add back FKeyDataSnapshot::Keys together with the saved positions and attributes.
+	 * On redo, remove the FKeyDataSnapshot::Keys.
+	 */
+	TMap<FCurveModelID, FCurveKeyData> SavedCurveState;
+	
+	bool HasChanges() const { return !SavedCurveState.IsEmpty(); }
+	operator bool() const { return HasChanges(); }
+	SIZE_T GetAllocatedSize() const { return FCurveKeyData::GetAllocatedSize(SavedCurveState); }
+};
+}
diff --git a/Engine/Source/Editor/CurveEditor/Public/Modification/Keys/ECurveChangeFlags.h b/Engine/Source/Editor/CurveEditor/Public/Modification/Keys/ECurveChangeFlags.h
new file mode 100644
--- /dev/null
+++ b/Engine/Source/Editor/CurveEditor/Public/Modification/Keys/ECurveChangeFlags.h
@@ -0,0 +1,33 @@
+﻿// Copyright Epic Games, Inc. All Rights Reserved.
+
+#pragma once
+
+#include "HAL/Platform.h"
+#include "Misc/EnumClassFlags.h"
+
+namespace UE::CurveEditor
+{
+/** Flags about the types of changes that can be made to a curve. */
+enum class ECurveChangeFlags : uint8
+{
+	None,
+
+	/** Detect change of FKeyPosition for FKeyHandles that are present before and after the change. */
+	MoveKeys = 1 << 0,
+	/** Detect FKeyHandles that are present after the change, but were not present before the change. */
+	AddKeys = 1 << 1,
+	/** Detect FKeyHandles that are not present after the change, but were present before the change. */
+	RemoveKeys = 1 << 2,
+	/**
+	 * Detect FKeyAttributes that are different after the change.
+	 * The attribute changes are detected for FKeyHandles that are present before and after the change.
+	 */
+	KeyAttributes = 1 << 3,
+	/** Detect changes to the curve attributes, e.g. pre and post extrapolation, etc. */
+	CurveAttributes = 1 << 4,
+
+	KeyData = MoveKeys | AddKeys | RemoveKeys | KeyAttributes,
+	All = KeyData | CurveAttributes
+};
+ENUM_CLASS_FLAGS(ECurveChangeFlags);
+}
diff --git a/Engine/Source/Editor/CurveEditor/Public/Modification/Keys/GenericCurveChangeUtils.h b/Engine/Source/Editor/CurveEditor/Public/Modification/Keys/GenericCurveChangeUtils.h
new file mode 100644
--- /dev/null
+++ b/Engine/Source/Editor/CurveEditor/Public/Modification/Keys/GenericCurveChangeUtils.h
@@ -0,0 +1,33 @@
+﻿// Copyright Epic Games, Inc. All Rights Reserved.
+
+#pragma once
+
+#include "Containers/ContainersFwd.h"
+#include "Templates/UniquePtr.h"
+
+#define UE_API CURVEEDITOR_API
+
+class FCurveEditor;
+class FCurveModel;
+class IBufferedCurveModel;
+namespace UE::CurveEditor { struct FGenericCurveChangeData; }
+struct FCurveModelID;
+
+namespace UE::CurveEditor::GenericCurveChange
+{
+/** Applies InDeltaChange to InCurves (redo). */
+UE_API void ApplyChange(const TMap<FCurveModelID, TUniquePtr<FCurveModel>>& InCurves, const FGenericCurveChangeData& InDeltaChange,
+	const double InCurrentSliderPos = 0
+	);
+/** Reverts InDeltaChange from InCurves (undo). */
+UE_API void RevertChange(const TMap<FCurveModelID, TUniquePtr<FCurveModel>>& InCurves, const FGenericCurveChangeData& InDeltaChange,
+	const double InCurrentSliderPos = 0
+	);
+
+/** Applies InDeltaChange to InCurveEditor (redo). */
+UE_API void ApplyChange(const FCurveEditor& InCurveEditor, const FGenericCurveChangeData& InDeltaChange);
+/** Reverts InDeltaChange from InCurveEditor (undo). */
+UE_API void RevertChange(const FCurveEditor& InCurveEditor, const FGenericCurveChangeData& InDeltaChange);
+}
+
+#undef UE_API
\ No newline at end of file
diff --git a/Engine/Source/Editor/CurveEditor/Public/Modification/Keys/SnapshotPerfDebugUtils.h b/Engine/Source/Editor/CurveEditor/Public/Modification/Keys/SnapshotPerfDebugUtils.h
new file mode 100644
--- /dev/null
+++ b/Engine/Source/Editor/CurveEditor/Public/Modification/Keys/SnapshotPerfDebugUtils.h
@@ -0,0 +1,173 @@
+﻿// Copyright Epic Games, Inc. All Rights Reserved.
+
+#pragma once
+
+#include "Containers/Map.h"
+#include "Data/CurvesSnapshot.h"
+#include "HAL/Platform.h"
+#include "Internationalization/Text.h"
+#include "Misc/EnumClassFlags.h"
+
+struct FCurveModelID;
+
+namespace UE::CurveEditor
+{
+/** Utility for retrieving some common performance data about FCurveDiffingData (or a collection thereof). */
+template<typename T>
+class TCurvesSnapshotPerfInfo
+{
+	const TMap<T, FCurveDiffingData>& Data;
+public:
+
+	explicit TCurvesSnapshotPerfInfo(const TMap<T, FCurveDiffingData>& InMap UE_LIFETIMEBOUND) : Data(InMap) {}
+	explicit TCurvesSnapshotPerfInfo(const FCurvesSnapshot& InSnapshot UE_LIFETIMEBOUND) : Data(InSnapshot.CurveData) {}
+
+	SIZE_T GetTotalAllocatedSize() const;
+
+	SIZE_T NumKeys() const;
+	SIZE_T NumPositions() const;
+	SIZE_T NumAttributes() const;
+
+	SIZE_T GetAllocatedSize_Keys() const;
+	SIZE_T GetAllocatedSize_Positions() const;
+	SIZE_T GetAllocatedSize_Attributes() const;
+};
+
+enum class ECurvesSnapshotPerfFlags
+{
+	None,
+
+	TotalAllocatedSize = 1 << 0,
+	
+	Keys = 1 << 1,
+	Positions = 1 << 2,
+	Attributes = 1 << 3,
+
+	All = TotalAllocatedSize | Keys | Positions | Attributes
+};
+ENUM_CLASS_FLAGS(ECurvesSnapshotPerfFlags);
+
+/** Converts the specified info into a string. */
+template<typename T>
+FString DumpSnapshotPerfData(const TCurvesSnapshotPerfInfo<T> InData, ECurvesSnapshotPerfFlags InFlags = ECurvesSnapshotPerfFlags::All);
+	
+template<typename T>
+FString DumpSnapshotPerfData(const TMap<T, FCurveDiffingData>& InMap, ECurvesSnapshotPerfFlags InFlags = ECurvesSnapshotPerfFlags::All)
+{
+	return DumpSnapshotPerfData(TCurvesSnapshotPerfInfo<T>(InMap), InFlags);
+}
+inline FString DumpSnapshotPerfData(const FCurvesSnapshot& InSnapshot, ECurvesSnapshotPerfFlags InFlags = ECurvesSnapshotPerfFlags::All)
+{
+	return DumpSnapshotPerfData(InSnapshot.CurveData, InFlags);
+}
+}
+
+namespace UE::CurveEditor
+{
+template <typename T>
+SIZE_T TCurvesSnapshotPerfInfo<T>::GetTotalAllocatedSize() const
+{
+	return Algo::TransformAccumulate(Data, [](const TPair<T, FCurveDiffingData>& Pair)
+	{
+		return Pair.Value.GetAllocatedSize();
+	}, 0);
+}
+
+template <typename T>
+SIZE_T TCurvesSnapshotPerfInfo<T>::NumKeys() const
+{
+	return Algo::TransformAccumulate(Data, [](const TPair<T, FCurveDiffingData>& Pair)
+	{
+		return Pair.Value.KeyHandles.Num();
+	}, 0);
+}
+
+template <typename T>
+SIZE_T TCurvesSnapshotPerfInfo<T>::NumPositions() const
+{
+	return Algo::TransformAccumulate(Data, [](const TPair<T, FCurveDiffingData>& Pair)
+	{
+		return Pair.Value.KeyPositions.Num();
+	}, 0);
+}
+
+template <typename T>
+SIZE_T TCurvesSnapshotPerfInfo<T>::NumAttributes() const
+{
+	return Algo::TransformAccumulate(Data, [](const TPair<T, FCurveDiffingData>& Pair)
+	{
+		return Pair.Value.KeyAttributes.Num();
+	}, 0);
+}
+
+template <typename T>
+SIZE_T TCurvesSnapshotPerfInfo<T>::GetAllocatedSize_Keys() const
+{
+	return Algo::TransformAccumulate(Data, [](const TPair<T, FCurveDiffingData>& Pair)
+	{
+		return Pair.Value.KeyHandles.GetAllocatedSize();
+	}, 0);
+}
+
+template <typename T>
+SIZE_T TCurvesSnapshotPerfInfo<T>::GetAllocatedSize_Positions() const
+{
+	return Algo::TransformAccumulate(Data, [](const TPair<T, FCurveDiffingData>& Pair)
+	{
+		return Pair.Value.KeyPositions.GetAllocatedSize();
+	}, 0);
+}
+
+template <typename T>
+SIZE_T TCurvesSnapshotPerfInfo<T>::GetAllocatedSize_Attributes() const
+{
+	return Algo::TransformAccumulate(Data, [](const TPair<T, FCurveDiffingData>& Pair)
+	{
+		return Pair.Value.KeyAttributes.GetAllocatedSize();
+	}, 0);
+}
+
+template <typename T>
+FString DumpSnapshotPerfData(const TCurvesSnapshotPerfInfo<T> InData, ECurvesSnapshotPerfFlags InFlags)
+{
+	TStringBuilder<256> Buffer;
+	bool bNeedsSeparator = false;
+	const auto AddSeparator = [&bNeedsSeparator, &Buffer]
+	{
+		if (!bNeedsSeparator)
+		{
+			bNeedsSeparator = true;
+		}
+		else
+		{
+			Buffer << TEXT(", ");
+		}
+	};
+
+	if (EnumHasAnyFlags(InFlags, ECurvesSnapshotPerfFlags::TotalAllocatedSize))
+	{
+		AddSeparator();
+		Buffer << TEXT("Total Size: ") << *FText::AsMemory(InData.GetTotalAllocatedSize()).ToString();
+	}
+
+	if (EnumHasAnyFlags(InFlags, ECurvesSnapshotPerfFlags::Keys))
+	{
+		AddSeparator();
+		Buffer << InData.NumKeys() << TEXT(" keys (") << *FText::AsMemory(InData.GetAllocatedSize_Keys()).ToString() << TEXT(")");
+	}
+
+	if (EnumHasAnyFlags(InFlags, ECurvesSnapshotPerfFlags::Keys))
+	{
+		AddSeparator();
+		Buffer << InData.NumKeys() << TEXT(" positions (") << *FText::AsMemory(InData.GetAllocatedSize_Positions()).ToString() << TEXT(")");
+	}
+
+	if (EnumHasAnyFlags(InFlags, ECurvesSnapshotPerfFlags::Keys))
+	{
+		AddSeparator();
+		Buffer << InData.NumKeys() << TEXT(" attributes (") << *FText::AsMemory(InData.GetAllocatedSize_Attributes()).ToString() << TEXT(")");
+	}
+	
+	return FString(Buffer.ToString());
+}
+}
\ No newline at end of file
diff --git a/Engine/Source/Editor/CurveEditor/Public/Modification/Utils/ScopedChangeBase.h b/Engine/Source/Editor/CurveEditor/Public/Modification/Utils/ScopedChangeBase.h
--- a/Engine/Source/Editor/CurveEditor/Public/Modification/Utils/ScopedChangeBase.h
+++ b/Engine/Source/Editor/CurveEditor/Public/Modification/Utils/ScopedChangeBase.h
@@ -2,8 +2,11 @@
 
 #pragma once
 
-#include "CurveEditorSelection.h"
-#include "ScopedTransaction.h"
+#include "Internationalization/Text.h"
+#include "Templates/SharedPointer.h"
+#include "Templates/UnrealTemplate.h"
+
+class FCurveEditor;
 
 namespace UE::CurveEditor
 {
@@ -14,6 +17,10 @@ class FCurveEditorCommandChange;
  * It detects if constructed within a FScopedTransaction context: by calling IsCancelled(), you can handle the case of the root transaction being cancelled.
  */
 class FScopedChangeBase
+	// We don't expect this to be copied / moved.
+	// If needed in the future, you can remove this BUT go through the subclasses and make sure their (auto-generated) move constructors work.
+	// Const members, like the ones in this class and subclasses, cause the compiler to auto-generate copy constructor instead of move constructors.
+	: public FNoncopyable 
 {
 public:
 
diff --git a/Engine/Source/Editor/CurveEditor/Public/Modification/Utils/ScopedKeyChange.h b/Engine/Source/Editor/CurveEditor/Public/Modification/Utils/ScopedKeyChange.h
new file mode 100644
--- /dev/null
+++ b/Engine/Source/Editor/CurveEditor/Public/Modification/Utils/ScopedKeyChange.h
@@ -0,0 +1,99 @@
+﻿// Copyright Epic Games, Inc. All Rights Reserved.
+
+#pragma once
+
+#include "Misc/EnumClassFlags.h"
+#include "Modification/Keys/CurveChangeDiff.h"
+#include "ScopedChangeBase.h"
+#include "Templates/SharedPointer.h"
+
+struct FCurveModelID;
+class FCurveEditor;
+
+namespace UE::CurveEditor
+{
+enum class EScopedKeyChangeFlags : uint8
+{
+	None,
+
+	/** Revert the change if FScopedKeyChange or surrounding FScopedTransaction becomes cancelled. */
+	RevertOnCancel = 1 << 0,
+	
+	Default = RevertOnCancel
+};
+ENUM_CLASS_FLAGS(EScopedKeyChangeFlags);
+
+enum class ETrackedCurves : uint8 { All, Selected };
+	
+/**
+ * Starts a transaction: When the scope ends, checks whether keys were added, removed, moved, or its key attributes edited.
+ * The changes are added as an undo-able action.
+ *
+ * ========== Performance ==========
+ * This class exists for convenience, faster iteration times, and more concise code but introduces overhead by temporarily buffering the curves.
+ * If the overhead is too great, consider building the change directly, and appending it to the transaction manager manually.
+ * However, in practice, the overhead should not be a concern, and for simplicity, usage of FScopedKeyChange is encouraged and preferred.
+ * 
+ * Stress test 15th of July, 2025, Development Build,
+ * Sequencer: ControlRigExample -> CopyPaste_Level-> PossessableBindings_CopyPaste -> Select Possessable_FKCR_CopyPaste
+ *  - Set Interp Linear:	109.782 keys,	Undo Size: 5.582 MiB,		Diff buffer size: 2.955 MiB,	Diff time: 4.2 ms
+ *  - Move Keys:            109.782 keys,	Undo Size: 1.36 MiB,		Diff buffer size: 2.117 MiB,	Diff time: 5.3 ms
+ *	- Drag Tangents:        109.782 keys,	Undo Size: 5.582 MiB,		Diff buffer size: 2.955 MiB,	Diff time: 4.4 ms
+ *	- Blend values:         109.782 keys,	Undo Size: 6.583 MiB,		Diff buffer size: 4.654 MiB,	Diff time: 10.7 ms
+ *	- Delete Keys:          109.782 keys,	Undo Size: 5.395 MiB,		Diff buffer size: 4.63 MiB,		Diff time: 1.8 ms
+ * These are the stats logged by CurveEditor.LogDiffPerfStats 1.
+ *
+ * ========== Combining FScopedKeyChange and FScopedSelectionChange ==========
+ * The following is operations supported:
+ * - Combination with FScopedSelectionChange (see below),
+ * - Nesting within FScopedTransaction; if the outer FScopedTransaction is cancelled, FScopedKeyChange can detect this (@see EScopedKeyChangeFlags).
+ * - Nested FScopedKeyChanges
+ * 
+ * When combining FScopedKeyChange and FScopedSelectionChange, order matters.
+ * Rule of thumb: Declare in the order the undo operation should revert; equivalent: declare in reverse order the redo operation should apply.
+ * 
+* Keep the following principles in mind:
+ * - C++ will deconstruct in reverse order
+ * - Unreal's transaction system execution order for custom changes is as follows:
+ *	- Revert (undo) in insertion order. Below, the call order would be FSelectionChangeCommand::Revert, then FGenericCurveChangeCommand::Revert.
+ *	- Apply (redo) in reverse insertion order. Below, the call order would be FGenericCurveChangeCommand::Apply, then FSelectionChangeCommand::Apply.
+ * 
+ * Example code for inserting a key and selecting it:
+ *		FScopedSelectionChange SelectionChange; // When undoing, first remove from selection...
+ *		FScopedKeyChange KeyChange; // ... then remove the key from the curve
+ *		// Add key to some curve
+ *		// Select the added key
+ *		//~KeyChange > adds key change as undo action
+ *		//~SelectionChange > adds selection change as undo action
+ * If you changed the declaration order, on redo, we'd first add to selection and then add the key to the curve.
+ * Redo would no longer work: as FSelectionChangeCommand would add the FKeyHandle to the selection first, the selection code would defensively check
+ * that a key exists first and reject the function call as the key has not yet been added to the curve model.
+ */
+class FScopedKeyChange : public FScopedChangeBase
+{
+public:
+	
+	/**
+	 * @param InCurveEditor The curve
+	 * @param InKeyChangeDiff Contains the data that is supposed to be diffed against
+	 * @param InDescription Optional description that will be used for the transaction title. This title does not matter if you have a surrounding FScopedTransaction.
+	 * @param InFlags Extra flags to affect the behaviour of the change
+	 */
+	[[nodiscard]] CURVEEDITOR_API explicit FScopedKeyChange(
+		TWeakPtr<FCurveEditor> InCurveEditor, FCurveChangeDiff InKeyChangeDiff, const FText& InDescription, EScopedKeyChangeFlags InFlags
+		);
+	
+	[[nodiscard]] CURVEEDITOR_API explicit FScopedKeyChange(
+		FCurvesSnapshotBuilder InDataToDiff, const FText& InDescription = FText::GetEmpty(), EScopedKeyChangeFlags InFlags = EScopedKeyChangeFlags::Default
+		);
+
+	CURVEEDITOR_API ~FScopedKeyChange();
+
+private:
+	
+	/** Checks for changes. */
+	const FCurveChangeDiff ChangeDiff;
+
+	const EScopedKeyChangeFlags Flags;
+};
+}
diff --git a/Engine/Source/Editor/CurveEditor/Public/Modification/Utils/ScopedSelectionChange.h b/Engine/Source/Editor/CurveEditor/Public/Modification/Utils/ScopedSelectionChange.h
--- a/Engine/Source/Editor/CurveEditor/Public/Modification/Utils/ScopedSelectionChange.h
+++ b/Engine/Source/Editor/CurveEditor/Public/Modification/Utils/ScopedSelectionChange.h
@@ -3,30 +3,55 @@
 #pragma once
 
 #include "CurveEditorSelection.h"
+#include "Misc/EnumClassFlags.h"
 #include "ScopedChangeBase.h"
 
 #define UE_API CURVEEDITOR_API
 
 namespace UE::CurveEditor
 {
+enum class EScopedSelectionChangeFlags : uint8
+{
+	None,
+
+	/** Revert the change if FScopedSelectionChange or surrounding FScopedTransaction becomes cancelled. */
+	RevertOnCancel = 1 << 0,
+
+	Default = RevertOnCancel
+};
+ENUM_CLASS_FLAGS(EScopedSelectionChangeFlags);
+	
 /**
  * Saves the current selection state and on destruction compares it against the then current selection. If the selection has changed, it appends
  * an undo-able action to the curve editor transaction manager.
  *
  * You can safely use this within the scope of a FScopedTransaction. In that case, the command will be added as sub-transaction to the parent
  * transaction.
+ *
+* The following is operations are supported:
+ * - Combination with FScopedKeyChange (@see FScopedKeyChange for detailed explanation).
+ * - Nesting within FScopedTransaction; if the outer FScopedTransaction is cancelled, FScopedSelectionChange can detect this (@see EScopedSelectionChangeFlags).
+ * - Nested FScopedSelectionChange
  */
 class FScopedSelectionChange : public FScopedChangeBase
 {
 public:
 
-	[[nodiscard]] UE_API explicit FScopedSelectionChange(TWeakPtr<FCurveEditor> InCurveEditor);
-	[[nodiscard]] UE_API explicit FScopedSelectionChange(TWeakPtr<FCurveEditor> InCurveEditor, const FText& InDescription);
+	[[nodiscard]] UE_API explicit FScopedSelectionChange(
+		TWeakPtr<FCurveEditor> InCurveEditor, EScopedSelectionChangeFlags InFlags = EScopedSelectionChangeFlags::Default
+		);
+	[[nodiscard]] UE_API explicit FScopedSelectionChange(
+		TWeakPtr<FCurveEditor> InCurveEditor, const FText& InDescription,
+		EScopedSelectionChangeFlags = EScopedSelectionChangeFlags::Default
+		);
 	
 	UE_API ~FScopedSelectionChange();
 
 private:
 
+	/** Flags that affect the behaviour of this change. */
+	const EScopedSelectionChangeFlags Flags;
+	
 	/** The selection the editor had when the transaction was started.  */
 	FCurveEditorSelection OriginalSelection;
 };
diff --git a/Engine/Source/Runtime/Engine/Classes/Curves/RealCurve.h b/Engine/Source/Runtime/Engine/Classes/Curves/RealCurve.h
--- a/Engine/Source/Runtime/Engine/Classes/Curves/RealCurve.h
+++ b/Engine/Source/Runtime/Engine/Classes/Curves/RealCurve.h
@@ -96,6 +96,15 @@ struct FCurveAttributes
 		return NewAttributes;
 	}
 
+	friend bool operator==(const FCurveAttributes& Left, const FCurveAttributes& Right)
+	{
+		return Left.bHasPreExtrapolation == Right.bHasPreExtrapolation
+			&& Left.bHasPostExtrapolation == Right.bHasPostExtrapolation
+			&& (!Left.bHasPreExtrapolation || Left.PreExtrapolation == Right.PreExtrapolation)
+			&& (!Left.bHasPostExtrapolation || Left.PostExtrapolation == Right.PostExtrapolation);
+	}
+	friend bool operator!=(const FCurveAttributes& Left, const FCurveAttributes& Right) { return !(Left == Right); }
+
 private:
 
 	/** true if the curve can express pre-extrapolation modes */
