Curve Editor: Improve default implementation of FCurveModel::GetAllKeys to avoid call to GetTimeRange & GetValueRange altogether.

diff --git a/Engine/Source/Editor/CurveEditor/Private/CurveModel.cpp b/Engine/Source/Editor/CurveEditor/Private/CurveModel.cpp
--- a/Engine/Source/Editor/CurveEditor/Private/CurveModel.cpp
+++ b/Engine/Source/Editor/CurveEditor/Private/CurveModel.cpp
@@ -6,13 +6,14 @@
 #include "CurveDataAbstraction.h"
 #include "Curves/KeyHandle.h"
 #include "HAL/PlatformCrt.h"
+#include "Math/NumericLimits.h"
 
 void FCurveModel::GetAllKeys(TArray<FKeyHandle>& OutKeyHandles) const
 {
-	double MinX, MaxX, MinY, MaxY;
-	GetTimeRange(MinX, MaxX);
-	GetValueRange(MinY, MaxY);
-	GetKeys(MinX, MaxX, MinY, MaxY, OutKeyHandles);
+	GetKeys(
+		TNumericLimits<double>::Lowest(), TNumericLimits<double>::Max(), TNumericLimits<double>::Lowest(), TNumericLimits<double>::Max(),
+		OutKeyHandles
+		);
 }
 
 void FCurveModel::PutKeys(
