UE-291476 Fix tool state in Curve Editor View when losing capture outside of the mousemove->Drag path

#tests Tested alt-tabbing mouse release outside with key dragging in the editor view. Tested rectangular selection tool in the same context
#rnx

diff --git a/Engine/Source/Editor/CurveEditor/Private/Views/SInteractiveCurveEditorView.cpp b/Engine/Source/Editor/CurveEditor/Private/Views/SInteractiveCurveEditorView.cpp
--- a/Engine/Source/Editor/CurveEditor/Private/Views/SInteractiveCurveEditorView.cpp
+++ b/Engine/Source/Editor/CurveEditor/Private/Views/SInteractiveCurveEditorView.cpp
@@ -90,6 +90,14 @@ namespace CurveViewConstants
 	constexpr float HoverProximityThresholdPx = 5.f;
 }
 
+#if WITH_EDITOR
+#include "Editor/EditorEngine.h"
+#include "EngineUtils.h"
+
+extern UNREALED_API UEditorEngine* GEditor;
+#endif // WITH_EDITOR
+
+
 #define LOCTEXT_NAMESPACE "SInteractiveCurveEditorView"
 
 TUniquePtr<ICurveEditorKeyDragOperation> CreateKeyDrag(ECurvePointType KeyType)
@@ -1752,6 +1760,20 @@ FReply SInteractiveCurveEditorView::OnMouseButtonDown(const FGeometry& MyGeometr
 	return FReply::Unhandled();
 }
 
+void SInteractiveCurveEditorView::OnMouseCaptureLost(const FCaptureLostEvent& CaptureLostEvent)
+{
+	// Delaying nulling out until next tick because this could be invoked during OnMouseMove()
+#if WITH_EDITOR
+	GEditor->GetTimerManager()->SetTimerForNextTick([this]()
+	{
+		if (!HasMouseCapture())
+		{
+			DragOperation.Reset();
+		}
+	});
+#endif
+}
+
 FReply SInteractiveCurveEditorView::OnMouseButtonUp(const FGeometry& MyGeometry, const FPointerEvent& MouseEvent)
 {
 	TSharedPtr<FCurveEditor> CurveEditor = WeakCurveEditor.Pin();
diff --git a/Engine/Source/Editor/CurveEditor/Public/Views/SInteractiveCurveEditorView.h b/Engine/Source/Editor/CurveEditor/Public/Views/SInteractiveCurveEditorView.h
--- a/Engine/Source/Editor/CurveEditor/Public/Views/SInteractiveCurveEditorView.h
+++ b/Engine/Source/Editor/CurveEditor/Public/Views/SInteractiveCurveEditorView.h
@@ -134,6 +134,9 @@ protected:
 	UE_API virtual FReply OnMouseWheel(const FGeometry& MyGeometry, const FPointerEvent& MouseEvent) override;
 	UE_API virtual FReply OnMouseButtonUp(const FGeometry& MyGeometry, const FPointerEvent& MouseEvent) override;
 	UE_API virtual FReply OnMouseButtonDown(const FGeometry& MyGeometry, const FPointerEvent& MouseEvent) override;
+
+public:
+	UE_API virtual void OnMouseCaptureLost(const FCaptureLostEvent& CaptureLostEvent) override;
 	// ~SWidget Interface
 
 protected:
