Curve Editor: Change transaction-based undo to new new command-based undo.

diff --git a/Engine/Plugins/Animation/ControlRig/Source/ControlRigEditor/Private/EditMode/Tween/ControlRigTweenModels.cpp b/Engine/Plugins/Animation/ControlRig/Source/ControlRigEditor/Private/EditMode/Tween/ControlRigTweenModels.cpp
--- a/Engine/Plugins/Animation/ControlRig/Source/ControlRigEditor/Private/EditMode/Tween/ControlRigTweenModels.cpp
+++ b/Engine/Plugins/Animation/ControlRig/Source/ControlRigEditor/Private/EditMode/Tween/ControlRigTweenModels.cpp
@@ -33,8 +33,9 @@ static TArray<TweeningUtilsEditor::FTweenModelUIEntry> MakeTweenModels(
 	{
 		using FModel = TEditorTweenModel<TControlRigTweenModel<TBlendFunction>>;
 		TUniquePtr<FModel> TweenModel = MakeUnique<FModel>(
-			CurveEditorAttr,										// TTangentFlatteningTweenProxy arg
-			InSequencerAttr, InOwningEditMode						// FControlRigTweenModel args
+			CurveEditorAttr,										// TTransactionalTweenModelProxy
+			CurveEditorAttr,										// TTangentFlatteningTweenProxy
+			InSequencerAttr, InOwningEditMode						// FControlRigTweenModel
 			);
 		Result.Emplace(MoveTemp(TweenModel), FTweenModelDisplayInfo(TBlendFunction));
 	};
@@ -47,8 +48,9 @@ static TArray<TweeningUtilsEditor::FTweenModelUIEntry> MakeTweenModels(
 	AddCurveTweenable.template operator()<EBlendFunction::BlendRelative>();
 	Result.Emplace(
 		MakeUnique<TEditorTweenModel<FCurveTimeOffsetTweenModel>>(
-			CurveEditorAttr,		// TTangentFlatteningTweenProxy arg
-			CurveEditorAttr			// FCurveTimeOffsetTweenModel args
+			CurveEditorAttr,		// TTransactionalTweenModelProxy
+			CurveEditorAttr,		// TTangentFlatteningTweenProxy
+			CurveEditorAttr			// FCurveTimeOffsetTweenModel
 			),
 		FTweenModelDisplayInfo(EBlendFunction::TimeOffset)
 	);
diff --git a/Engine/Plugins/Animation/TweeningUtils/Source/TweeningUtilsEditor/Public/Math/CurveBlending.h b/Engine/Plugins/Animation/TweeningUtils/Source/TweeningUtilsEditor/Public/Math/CurveBlending.h
--- a/Engine/Plugins/Animation/TweeningUtils/Source/TweeningUtilsEditor/Public/Math/CurveBlending.h
+++ b/Engine/Plugins/Animation/TweeningUtils/Source/TweeningUtilsEditor/Public/Math/CurveBlending.h
@@ -88,7 +88,6 @@ bool BlendCurve_ByKeyRange(const FCurveEditor& InCurveEditor, const FContiguousK
 		{
 			continue;
 		}
-		Curve->Modify();
 		
 		const FContiguousKeyMapping::FContiguousKeysArray& CurveBlendedKeys = KeysArray.Value;
 		for (const FContiguousKeys& BlendRange : CurveBlendedKeys.KeysArray)
diff --git a/Engine/Plugins/Animation/TweeningUtils/Source/TweeningUtilsEditor/Public/Math/Models/TransactionalTweenModelProxy.h b/Engine/Plugins/Animation/TweeningUtils/Source/TweeningUtilsEditor/Public/Math/Models/TransactionalTweenModelProxy.h
--- a/Engine/Plugins/Animation/TweeningUtils/Source/TweeningUtilsEditor/Public/Math/Models/TransactionalTweenModelProxy.h
+++ b/Engine/Plugins/Animation/TweeningUtils/Source/TweeningUtilsEditor/Public/Math/Models/TransactionalTweenModelProxy.h
@@ -4,13 +4,16 @@
 
 #include "TweenModel.h"
 
+#include "CurveEditor.h"
 #include "Internationalization/Internationalization.h"
 #include "Internationalization/Text.h"
-#include "ScopedTransaction.h"
+#include "Misc/Attribute.h"
+#include "Modification/Utils/ScopedKeyChange.h"
 #include "Templates/UniquePtr.h"
 
 #include <type_traits>
 
+
 class FCurveEditor;
 
 namespace UE::TweeningUtilsEditor
@@ -22,12 +25,21 @@ class TTransactionalTweenModelProxy : public TBase
 public:
 
 	template<typename... TArg>
-	explicit TTransactionalTweenModelProxy(TArg&&... Arg) : TBase(Forward<TArg>(Arg)...) {}
+	explicit TTransactionalTweenModelProxy(TAttribute<TWeakPtr<FCurveEditor>> InCurveEditorAttr, TArg&&... Arg)
+		: TBase(Forward<TArg>(Arg)...)
+		, CurveEditorAttr(MoveTemp(InCurveEditorAttr))
+	{}
 	
 	//~ Begin FTweenModel Interface
 	virtual void StartBlendOperation() override
 	{
-		InProgressTransaction = MakeUnique<FScopedTransaction>(NSLOCTEXT("FTransactionalTweenModelProxy", "Transaction", "Blend values"));
+		if (const TSharedPtr<FCurveEditor> CurveEditor = CurveEditorAttr.Get().Pin())
+		{
+			InProgressTransaction.Emplace(
+				CurveEditor::FCurvesSnapshotBuilder(CurveEditor).TrackAllCurves(),
+				NSLOCTEXT("FTransactionalTweenModelProxy", "Transaction", "Blend values")
+				);
+		}
 		TBase::StartBlendOperation();
 	}
 	virtual void StopBlendOperation() override
@@ -39,8 +51,11 @@ public:
 
 private:
 
+	/** Used to detect key changes. */
+	const TAttribute<TWeakPtr<FCurveEditor>> CurveEditorAttr;
+	
 	/** Active during a blend operation. */
-	TUniquePtr<FScopedTransaction> InProgressTransaction;
+	TOptional<CurveEditor::FScopedKeyChange> InProgressTransaction;
 };
 }
 
diff --git a/Engine/Plugins/Editor/CurveEditorTools/Source/CurveEditorTools/Private/Extensions/Tweening/CurveEditorTweenModels.cpp b/Engine/Plugins/Editor/CurveEditorTools/Source/CurveEditorTools/Private/Extensions/Tweening/CurveEditorTweenModels.cpp
--- a/Engine/Plugins/Editor/CurveEditorTools/Source/CurveEditorTools/Private/Extensions/Tweening/CurveEditorTweenModels.cpp
+++ b/Engine/Plugins/Editor/CurveEditorTools/Source/CurveEditorTools/Private/Extensions/Tweening/CurveEditorTweenModels.cpp
@@ -19,8 +19,9 @@ static TArray<TweeningUtilsEditor::FTweenModelUIEntry> MakeTweenModels(TWeakPtr<
 	{
 		using FModel = TEditorTweenModel<TContextAwareTweeningModel<TBlendFunction>>;
 		TUniquePtr<FModel> TweenModel = MakeUnique<FModel>(
-			InCurveEditor,			// TTangentFlatteningTweenProxy arg
-			InCurveEditor			// FContextAwareTweeningModel args
+			InCurveEditor,			// TTransactionalTweenModelProxy
+			InCurveEditor,			// TTangentFlatteningTweenProxy
+			InCurveEditor			// FContextAwareTweeningModel
 			);
 		Result.Emplace(MoveTemp(TweenModel), FTweenModelDisplayInfo(TBlendFunction));
 	};
@@ -33,8 +34,9 @@ static TArray<TweeningUtilsEditor::FTweenModelUIEntry> MakeTweenModels(TWeakPtr<
 	AddCurveTweenable.template operator()<EBlendFunction::BlendRelative>();
 	Result.Emplace(
 		MakeUnique<TEditorTweenModel<FCurveTimeOffsetTweenModel>>(
-			InCurveEditor,		// TTangentFlatteningTweenProxy arg
-			InCurveEditor		// FCurveTimeOffsetTweenModel args
+			InCurveEditor,			// TTransactionalTweenModelProxy
+			InCurveEditor,		// TTangentFlatteningTweenProxy
+			InCurveEditor		// FCurveTimeOffsetTweenModel
 			),
 		FTweenModelDisplayInfo(EBlendFunction::TimeOffset)
 	);
diff --git a/Engine/Plugins/Editor/CurveEditorTools/Source/CurveEditorTools/Private/Tools/CurveEditorMultiScaleTool.cpp b/Engine/Plugins/Editor/CurveEditorTools/Source/CurveEditorTools/Private/Tools/CurveEditorMultiScaleTool.cpp
--- a/Engine/Plugins/Editor/CurveEditorTools/Source/CurveEditorTools/Private/Tools/CurveEditorMultiScaleTool.cpp
+++ b/Engine/Plugins/Editor/CurveEditorTools/Source/CurveEditorTools/Private/Tools/CurveEditorMultiScaleTool.cpp
@@ -10,6 +10,8 @@
 #include "SCurveEditorPanel.h"
 #include "Algo/MaxElement.h"
 #include "Misc/Mirror/TangentSelectionFlattener.h"
+#include "Modification/Keys/CurvesSnapshotBuilder.h"
+#include "Modification/Utils/ScopedKeyChange.h"
 
 #include UE_INLINE_GENERATED_CPP_BY_NAME(CurveEditorMultiScaleTool)
 
@@ -378,14 +380,17 @@ FSlateIcon FCurveEditorMultiScaleTool::GetIcon() const
 
 void FCurveEditorMultiScaleTool::OnDragStart()
 {
-	ActiveTransaction = MakeUnique<FScopedTransaction>(TEXT("CurveEditorMultiScaleTool"), LOCTEXT("CurveEditorMultiScaleToolTransaction", "Multi Scale Key(s)"), nullptr);
-
 	TSharedPtr<FCurveEditor> CurveEditor = WeakCurveEditor.Pin();
 	if (!CurveEditor.IsValid())
 	{
 		return;
 	}
-	
+
+	using namespace UE::CurveEditor;
+	ActiveTransaction = MakeUnique<FScopedKeyChange>(
+		FCurvesSnapshotBuilder(CurveEditor, ECurveChangeFlags::KeyAttributes | ECurveChangeFlags::MoveKeys).TrackSelectedCurves(),
+		LOCTEXT("CurveEditorMultiScaleToolTransaction", "Multi Scale Key(s)")
+		);
 	TangentFlattener.Emplace(*CurveEditor);
 
 	// We need to cache our key data because all of our calculations have to be relative to the starting data and not the current per-frame data.
@@ -399,8 +404,6 @@ void FCurveEditorMultiScaleTool::OnDragStart()
 
 		if (ensureAlways(Curve))
 		{
-			Curve->Modify();
-
 			TArrayView<const FKeyHandle> Handles = Pair.Value.AsArray();
 
 			FKeyData& KeyData = KeysByCurve.Emplace_GetRef(CurveID);
diff --git a/Engine/Plugins/Editor/CurveEditorTools/Source/CurveEditorTools/Private/Tools/CurveEditorMultiScaleTool.h b/Engine/Plugins/Editor/CurveEditorTools/Source/CurveEditorTools/Private/Tools/CurveEditorMultiScaleTool.h
--- a/Engine/Plugins/Editor/CurveEditorTools/Source/CurveEditorTools/Private/Tools/CurveEditorMultiScaleTool.h
+++ b/Engine/Plugins/Editor/CurveEditorTools/Source/CurveEditorTools/Private/Tools/CurveEditorMultiScaleTool.h
@@ -2,12 +2,12 @@
 
 #pragma once
 
-#include "ICurveEditorToolExtension.h"
-#include "Framework/DelayedDrag.h"
 #include "CurveDataAbstraction.h"
 #include "CurveEditorTypes.h"
-#include "ScopedTransaction.h"
+#include "Framework/DelayedDrag.h"
+#include "ICurveEditorToolExtension.h"
 #include "Misc/Mirror/TangentSelectionFlattener.h"
+#include "Modification/Utils/ScopedKeyChange.h"
 #include "UObject/StructOnScope.h"
 #include "CurveEditorMultiScaleTool.generated.h"
 
@@ -133,7 +133,7 @@ private:
 	TWeakPtr<FCurveEditor> WeakCurveEditor;
 
 	/** The currently open transaction (if any) */
-	TUniquePtr<class FScopedTransaction> ActiveTransaction;
+	TUniquePtr<UE::CurveEditor::FScopedKeyChange> ActiveTransaction;
 	
 	/** Set when attempting to move a drag handle. */
 	TOptional<FDelayedDrag> DelayedDrag;
diff --git a/Engine/Plugins/Editor/CurveEditorTools/Source/CurveEditorTools/Private/Tools/CurveEditorRetimeTool.cpp b/Engine/Plugins/Editor/CurveEditorTools/Source/CurveEditorTools/Private/Tools/CurveEditorRetimeTool.cpp
--- a/Engine/Plugins/Editor/CurveEditorTools/Source/CurveEditorTools/Private/Tools/CurveEditorRetimeTool.cpp
+++ b/Engine/Plugins/Editor/CurveEditorTools/Source/CurveEditorTools/Private/Tools/CurveEditorRetimeTool.cpp
@@ -13,6 +13,7 @@
 #include "Rendering/SlateRenderer.h"
 #include "Fonts/FontMeasure.h"
 #include "SCurveEditorView.h"
+#include "Modification/Utils/ScopedKeyChange.h"
 #include "Styling/CoreStyle.h"
 
 #include UE_INLINE_GENERATED_CPP_BY_NAME(CurveEditorRetimeTool)
@@ -479,11 +480,9 @@ FSlateIcon FCurveEditorRetimeTool::GetIcon() const
 
 void FCurveEditorRetimeTool::OnDragStart()
 {
-	ActiveTransaction = MakeUnique<FScopedTransaction>(TEXT("CurveEditorRetimeTool"), LOCTEXT("CurveEditorRetimeToolTransaction", "Retime Key(s)"), nullptr);
-
-	TSharedPtr<FCurveEditor> CurveEditor = WeakCurveEditor.Pin();
-	if (CurveEditor)
+	if (const TSharedPtr<FCurveEditor> CurveEditor = WeakCurveEditor.Pin())
 	{
+		ActiveTransaction.Emplace(CurveEditor, LOCTEXT("CurveEditorRetimeToolTransaction", "Retime Key(s)"));
 		// Store our Retime Data so the handles reset to the correct location upon undo/redo.
 		RetimeData->Modify();
 
@@ -505,8 +504,6 @@ void FCurveEditorRetimeTool::OnDragStart()
 			FCurveModel* CurveModel = CurveEditor->FindCurve(CurveID);
 			check(CurveModel);
 
-			CurveModel->Modify();
-
 			FPreDragChannelData& ChannelData = ResizeData.CurveChannels[ResizeData.CurveChannels.Emplace(CurveID)];
 
 			// Get all Key Handles
diff --git a/Engine/Plugins/Editor/CurveEditorTools/Source/CurveEditorTools/Private/Tools/CurveEditorRetimeTool.h b/Engine/Plugins/Editor/CurveEditorTools/Source/CurveEditorTools/Private/Tools/CurveEditorRetimeTool.h
--- a/Engine/Plugins/Editor/CurveEditorTools/Source/CurveEditorTools/Private/Tools/CurveEditorRetimeTool.h
+++ b/Engine/Plugins/Editor/CurveEditorTools/Source/CurveEditorTools/Private/Tools/CurveEditorRetimeTool.h
@@ -1,10 +1,13 @@
 // Copyright Epic Games, Inc. All Rights Reserved.
 
 #pragma once
+
 #include "ICurveEditorToolExtension.h"
 #include "Framework/DelayedDrag.h"
 #include "CurveDataAbstraction.h"
 #include "CurveEditorSelection.h"
+#include "Modification/Utils/ScopedKeyChange.h"
+#include "ScopedTransaction.h"
 #include "CurveEditorRetimeTool.generated.h"
 
 class FCurveEditor;
@@ -103,8 +106,22 @@ private:
 	/** The retiming data for this tool. This is a UObject for Undo/Redo purposes. Created and added to the root set on construction of the tool. */
 	UCurveEditorRetimeToolData* RetimeData;
 
+	struct FTransaction
+	{
+		FScopedTransaction Transaction;
+		UE::CurveEditor::FScopedKeyChange KeyChange;
+
+		explicit FTransaction(TWeakPtr<FCurveEditor> InCurveEditor, const FText& Description)
+			: Transaction(Description)
+			, KeyChange([InCurveEditor]
+			{
+				using namespace UE::CurveEditor;
+				return FCurvesSnapshotBuilder(InCurveEditor, ECurveChangeFlags::KeyAttributes | ECurveChangeFlags::MoveKeys).TrackAllCurves();
+			}())
+		{}
+	};
 	/** The currently open transaction (if any) */
-	TUniquePtr<class FScopedTransaction> ActiveTransaction;
+	TOptional<FTransaction> ActiveTransaction;
 	
 	/** Set when attempting to move a drag handle. */
 	TOptional<FDelayedDrag> DelayedDrag;
diff --git a/Engine/Plugins/Editor/CurveEditorTools/Source/CurveEditorTools/Private/Tools/CurveEditorTransformTool.cpp b/Engine/Plugins/Editor/CurveEditorTools/Source/CurveEditorTools/Private/Tools/CurveEditorTransformTool.cpp
--- a/Engine/Plugins/Editor/CurveEditorTools/Source/CurveEditorTools/Private/Tools/CurveEditorTransformTool.cpp
+++ b/Engine/Plugins/Editor/CurveEditorTools/Source/CurveEditorTools/Private/Tools/CurveEditorTransformTool.cpp
@@ -786,7 +786,11 @@ void FCurveEditorTransformTool::OnDragStart()
 		return;
 	}
 
-	ActiveTransaction = MakeUnique<FScopedTransaction>(TEXT("CurveEditorTransformTool"), LOCTEXT("CurveEditorTransformToolTransaction", "Transform Key(s)"), nullptr);
+	using namespace UE::CurveEditor;
+	ActiveTransaction = MakeUnique<FScopedKeyChange>(
+		FCurvesSnapshotBuilder(CurveEditor, ECurveChangeFlags::KeyAttributes | ECurveChangeFlags::MoveKeys).TrackSelectedCurves(),
+		LOCTEXT("CurveEditorTransformToolTransaction", "Transform Key(s)")
+		);
 	TangentFlattener.Emplace(*CurveEditor);
 
 	CurveEditor->SuppressBoundTransformUpdates(true);
@@ -802,8 +806,6 @@ void FCurveEditorTransformTool::OnDragStart()
 
 		if (ensureAlways(Curve))
 		{
-			Curve->Modify();
-
 			TArrayView<const FKeyHandle> Handles = Pair.Value.AsArray();
 
 			FKeyData& KeyData = KeysByCurve.Emplace_GetRef(CurveID);
diff --git a/Engine/Plugins/Editor/CurveEditorTools/Source/CurveEditorTools/Private/Tools/CurveEditorTransformTool.h b/Engine/Plugins/Editor/CurveEditorTools/Source/CurveEditorTools/Private/Tools/CurveEditorTransformTool.h
--- a/Engine/Plugins/Editor/CurveEditorTools/Source/CurveEditorTools/Private/Tools/CurveEditorTransformTool.h
+++ b/Engine/Plugins/Editor/CurveEditorTools/Source/CurveEditorTools/Private/Tools/CurveEditorTransformTool.h
@@ -9,6 +9,7 @@
 #include "ScopedTransaction.h"
 #include "CurveEditorSnapMetrics.h"
 #include "Misc/Mirror/TangentSelectionFlattener.h"
+#include "Modification/Utils/ScopedKeyChange.h"
 #include "CurveEditorTransformTool.generated.h"
 
 class FCurveEditor;
@@ -156,7 +157,7 @@ private:
 	TWeakPtr<FCurveEditor> WeakCurveEditor;
 
 	/** The currently open transaction (if any) */
-	TUniquePtr<class FScopedTransaction> ActiveTransaction;
+	TUniquePtr<UE::CurveEditor::FScopedKeyChange> ActiveTransaction;
 
 	/** Cached information about our transform tool such as interaction state, etc. */
 	FCurveEditorTransformWidget TransformWidget;
diff --git a/Engine/Plugins/Editor/CurveEditorTools/Source/CurveEditorTools/Private/Tools/Lattice/CurveEditorLatticeTool.cpp b/Engine/Plugins/Editor/CurveEditorTools/Source/CurveEditorTools/Private/Tools/Lattice/CurveEditorLatticeTool.cpp
--- a/Engine/Plugins/Editor/CurveEditorTools/Source/CurveEditorTools/Private/Tools/Lattice/CurveEditorLatticeTool.cpp
+++ b/Engine/Plugins/Editor/CurveEditorTools/Source/CurveEditorTools/Private/Tools/Lattice/CurveEditorLatticeTool.cpp
@@ -110,6 +110,7 @@ FReply FCurveEditorLatticeTool::OnMouseButtonDoubleClick(
 	TSharedRef<SWidget> OwningWidget, const FGeometry& InViewContainerGeometry, const FPointerEvent& InMouseEvent
 	)
 {
+	using namespace UE::CurveEditor;
 	CancelDrag();
 	
 	const bool bHandled = ComputeHoverStateForLeftClick(
@@ -120,7 +121,11 @@ FReply FCurveEditorLatticeTool::OnMouseButtonDoubleClick(
 			if (HoverState.HoveredCell)
 			{
 				const FScopedTransaction Transaction(LOCTEXT("FlipCurves", "Flip curves"));
+				const FScopedKeyChange KeyChange(
+					FCurvesSnapshotBuilder(WeakCurveEditor, ECurveChangeFlags::KeyAttributes | ECurveChangeFlags::MoveKeys).TrackSelectedCurves()
+					);
 				UndoObject->Modify();
+				
 				HandleMirrorGrid();
 				SaveUndoState(InCurveEditor);
 				return true;
@@ -128,7 +133,11 @@ FReply FCurveEditorLatticeTool::OnMouseButtonDoubleClick(
 			if (HoverState.HoveredEdge)
 			{
 				const FScopedTransaction Transaction(LOCTEXT("FlattenCurves", "Flatten curves"));
+				const FScopedKeyChange KeyChange(
+					FCurvesSnapshotBuilder(WeakCurveEditor, ECurveChangeFlags::KeyAttributes | ECurveChangeFlags::MoveKeys).TrackSelectedCurves()
+					);
 				UndoObject->Modify();
+				
 				HandleMoveEdgeToOpposite(*HoverState.HoveredEdge);
 				SaveUndoState(InCurveEditor);
 				return true;
@@ -526,7 +535,7 @@ TOptional<FCurveEditorLatticeTool::FDragOpData> FCurveEditorLatticeTool::CreateD
 		WeakCurveEditor, FControlPointArray{ ControlPoints[InControlPointIndex] }, InBounds, HandleDrag, HandleFinishDrag
 	);
 	FDragOpData DragDrop(
-		MoveTemp(Op),
+		WeakCurveEditor, MoveTemp(Op),
 		FText::Format(LOCTEXT("MoveControlPoint", "Move control point ({0})"), LexViewMode(WeakCurveEditor.Pin()->GetPanel()->GetViewMode()))
 		);
 
@@ -577,7 +586,7 @@ TOptional<FCurveEditorLatticeTool::FDragOpData> FCurveEditorLatticeTool::CreateD
 		WeakCurveEditor, MoveTemp(EdgePoints), InBounds, HandleDrag, HandleFinishDrag
 	);
 	FDragOpData DragDrop(
-		MoveTemp(Op),
+		WeakCurveEditor, MoveTemp(Op),
 		FText::Format(LOCTEXT("MoveEdge", "Move edge ({0})"), LexViewMode(WeakCurveEditor.Pin()->GetPanel()->GetViewMode()))
 		);
 
@@ -630,7 +639,7 @@ TOptional<FCurveEditorLatticeTool::FDragOpData> FCurveEditorLatticeTool::CreateD
 		WeakCurveEditor, MoveTemp(CellPoints), InBounds, HandleDrag, HandleFinishDrag
 	);
 	return FDragOpData(
-		MoveTemp(Op),
+		WeakCurveEditor, MoveTemp(Op),
 		FText::Format(LOCTEXT("MoveCell", "Move cell ({0})"), LexViewMode(WeakCurveEditor.Pin()->GetPanel()->GetViewMode()))
 		);
 }
@@ -792,7 +801,6 @@ void FCurveEditorLatticeTool::OnSetControlPoints(
 	// - SetKeyPositions 138.954 ms		- Total 139.229 ms
 	// - SetKeyPositions 269.523 ms		- Total 269.811 ms
 	// Conclusion: The following SetKeyPositions block is slow.
-	InCurveModel.Modify();
 	InCurveModel.SetKeyPositions(Handles, Keys, InChangeType);
 }
 
@@ -804,9 +812,6 @@ void FCurveEditorLatticeTool::MoveEdgeToOpposite(FPerCurveDeformer2D& InDeformer
 		return;
 	}
 	
-	// OnSetControlPoints modifies the curves too late - MirrorTangents will already have modified them.
-	InCurveModel.Modify();
-	
 	// Flatten tangents if moving top or bottom edge...
 	const FCurveEditor& CurveEditor = *WeakCurveEditor.Pin();
 	TOptional<FLatticeEdgeTangentsMirrorOp> MirrorOp = FLatticeEdgeTangentsMirrorOp::MakeMirrorOpForDragLatticeEdge(
@@ -833,8 +838,6 @@ void FCurveEditorLatticeTool::MirrorGrid(FPerCurveDeformer2D& InDeformer, FCurve
 		return;
 	}
 	
-	InCurveModel.Modify(); 
-	
 	const TConstArrayView<FKeyHandle> Keys = InDeformer.GetCellMetaData(0);
 	const auto[Min, Max] = CurveEditor::FindMinMaxHeight(Keys, InCurveModel);
 	const double Midpoint = Min + (Max - Min) * 0.5;
diff --git a/Engine/Plugins/Editor/CurveEditorTools/Source/CurveEditorTools/Private/Tools/Lattice/CurveEditorLatticeTool.h b/Engine/Plugins/Editor/CurveEditorTools/Source/CurveEditorTools/Private/Tools/Lattice/CurveEditorLatticeTool.h
--- a/Engine/Plugins/Editor/CurveEditorTools/Source/CurveEditorTools/Private/Tools/Lattice/CurveEditorLatticeTool.h
+++ b/Engine/Plugins/Editor/CurveEditorTools/Source/CurveEditorTools/Private/Tools/Lattice/CurveEditorLatticeTool.h
@@ -16,6 +16,7 @@
 #include "Misc/LatticeDrawUtils.h"
 #include "PerCurveLatticeData.h"
 #include "ScopedTransaction.h"
+#include "Modification/Utils/ScopedKeyChange.h"
 #include "Templates/SharedPointer.h"
 #include "UObject/UnrealType.h"
 
@@ -139,6 +140,8 @@ private:
 		TUniquePtr<FLatticeDragOp> DragOp;
 		/** This transaction is active for the duration of the drag. */
 		TUniquePtr<FScopedTransaction> Transaction;
+		/** Transactions the key changes */
+		TUniquePtr<CurveEditor::FScopedKeyChange> KeyChange;
 		
 		/** Adjust tangents while dragging on an edge.*/
 		TMap<FCurveModelID, FLatticeEdgeTangentsMirrorOp> EdgeTangentMirroringOps;
@@ -148,9 +151,16 @@ private:
 		/** False until the first drag actually changes keys. We only need to capture undo state just before the first change. */
 		bool bSavedUndoState = false;
 
-		explicit FDragOpData(TUniquePtr<FLatticeDragOp> DragOp, const FText& TransactionText)
+		explicit FDragOpData(TWeakPtr<FCurveEditor> InCurveEditor, TUniquePtr<FLatticeDragOp> DragOp, const FText& TransactionText)
 			: DragOp(MoveTemp(DragOp))
 			, Transaction(MakeUnique<FScopedTransaction>(TEXT("CurveEditorLatticeTool"), TransactionText, nullptr))
+			, KeyChange([InCurveEditor]
+			{
+				using namespace UE::CurveEditor;
+				return MakeUnique<FScopedKeyChange>(
+					FCurvesSnapshotBuilder(InCurveEditor, ECurveChangeFlags::MoveKeys | ECurveChangeFlags::KeyAttributes).TrackSelectedCurves()
+					);
+			}())
 		{}
 	};
 
diff --git a/Engine/Plugins/Editor/CurveEditorTools/Source/CurveEditorTools/Private/Tools/Lattice/Mirror/LatticeEdgeTangentMirrorOp.cpp b/Engine/Plugins/Editor/CurveEditorTools/Source/CurveEditorTools/Private/Tools/Lattice/Mirror/LatticeEdgeTangentMirrorOp.cpp
--- a/Engine/Plugins/Editor/CurveEditorTools/Source/CurveEditorTools/Private/Tools/Lattice/Mirror/LatticeEdgeTangentMirrorOp.cpp
+++ b/Engine/Plugins/Editor/CurveEditorTools/Source/CurveEditorTools/Private/Tools/Lattice/Mirror/LatticeEdgeTangentMirrorOp.cpp
@@ -87,7 +87,6 @@ void FLatticeEdgeTangentsMirrorOp::OnMoveEdge(TConstArrayView<FVector2D> InNewEd
 	{
 		return;
 	}
-	CurveModel->Modify();
 		
 	CurveMirrorData.Solver.OnMoveEdge(InCurveEditor, ChooseConsistentEdgeVert(*EdgeType, InNewEdge[0], InNewEdge[1]));
 }
diff --git a/Engine/Plugins/Editor/CurveEditorTools/Source/CurveEditorTools/Private/Tools/Lattice/Mirror/LatticePointTangentMirrorOp.cpp b/Engine/Plugins/Editor/CurveEditorTools/Source/CurveEditorTools/Private/Tools/Lattice/Mirror/LatticePointTangentMirrorOp.cpp
--- a/Engine/Plugins/Editor/CurveEditorTools/Source/CurveEditorTools/Private/Tools/Lattice/Mirror/LatticePointTangentMirrorOp.cpp
+++ b/Engine/Plugins/Editor/CurveEditorTools/Source/CurveEditorTools/Private/Tools/Lattice/Mirror/LatticePointTangentMirrorOp.cpp
@@ -55,7 +55,6 @@ void FLatticePointTangentsMirrorOp::OnMovePoint(const FVector2D& InControlPoint,
 	{
 		return;
 	}
-	CurveModel->Modify();
 		
 	CurveEditor::FCurveTangentMirrorData& MirrorData = CurveData.TangentMirrorSolver.CurveData[CurveId];
 	RecomputeMirroringParallel(CurveEditor, CurveId, MirrorData, InControlPoint.Y,
diff --git a/Engine/Source/Editor/CurveEditor/Private/CurveEditor.cpp b/Engine/Source/Editor/CurveEditor/Private/CurveEditor.cpp
--- a/Engine/Source/Editor/CurveEditor/Private/CurveEditor.cpp
+++ b/Engine/Source/Editor/CurveEditor/Private/CurveEditor.cpp
@@ -60,6 +60,7 @@
 #include "UObject/UnrealNames.h"
 #include "UnrealExporter.h"
 #include "Misc/SmartSnap.h"
+#include "Modification/Utils/ScopedKeyChange.h"
 #include "Modification/Utils/ScopedSelectionChange.h"
 #include "Widgets/Colors/SColorPicker.h"
 #include "Widgets/Notifications/SNotificationList.h"
@@ -1289,10 +1290,14 @@ void FCurveEditor::FlipCurveVertical(TArray<FKeyPosition>& AllKeyPositions, TArr
 
 void FCurveEditor::FlipCurve(ECurveFlipDirection Direction)
 {
-	FScopedTransaction Transaction(LOCTEXT("FlipCurve", "Flip Curve"));
-	for (FCurveModelID ID : GetEditedCurves())
+	using namespace UE::CurveEditor;
+	const FScopedKeyChange KeyChange(
+		FCurvesSnapshotBuilder(SharedThis(this), CurveData, ECurveChangeFlags::MoveKeys | ECurveChangeFlags::KeyAttributes),
+		LOCTEXT("FlipCurve", "Flip Curve")
+		);
+	for (const TPair<FCurveModelID, TUniquePtr<FCurveModel>>& Pair : CurveData)
 	{
-		if (FCurveModel* Curve = FindCurve(ID))
+		if (FCurveModel* Curve = Pair.Value.Get())
 		{
 			// Init key handles
 			TArray<FKeyHandle> KeyHandles;
@@ -1328,7 +1333,6 @@ void FCurveEditor::FlipCurve(ECurveFlipDirection Direction)
 
 			if (AllKeyPositions.Num() > 0)
 			{
-				Curve->Modify();
 				Curve->SetKeyPositions(KeyHandles, AllKeyPositions);
 				Curve->SetKeyAttributes(KeyHandles, AllKeyAttributes);
 			}
@@ -1336,8 +1340,7 @@ void FCurveEditor::FlipCurve(ECurveFlipDirection Direction)
 	}
 }
 
-void
-FCurveEditor::ToggleExpandCollapseNodes(bool bRecursive)
+void FCurveEditor::ToggleExpandCollapseNodes(bool bRecursive)
 {
 	Tree.ToggleExpansionState(bRecursive);
 }
@@ -1547,15 +1550,19 @@ public:
 
 void FCurveEditor::MatchLastTangentToFirst(bool bMatchLastToFirst)
 {
-	FScopedTransaction Transaction(LOCTEXT("MatchTangents", "Match Tangents"));
-	bool bFoundAnyTangents = false;
+	using namespace UE::CurveEditor;
+	// FScopedKeyChange won't generate any transaction if no tangents are changed
+	const FScopedKeyChange KeyChange(
+		FCurvesSnapshotBuilder(SharedThis(this), CurveData, ECurveChangeFlags::KeyAttributes),
+		LOCTEXT("MatchTangents", "Match Tangents")
+		);
 
-	TArray<FKeyHandle> KeyHandles;
+	TArray<FKeyHandle, TInlineAllocator<2>> KeyHandles;
 	KeyHandles.SetNum(2);
-	TArray<FKeyAttributes> KeyAttributes;
+	TArray<FKeyAttributes, TInlineAllocator<2>> KeyAttributes;
 	KeyAttributes.SetNum(2);
+	
 	TArray<FKeyHandle> AllKeyHandles;
-
 	for (const TTuple<FCurveModelID, TUniquePtr<FCurveModel>>& Pair : CurveData)
 	{
 		FCurveModel* CurveModel = Pair.Value.Get();
@@ -1569,7 +1576,6 @@ void FCurveEditor::MatchLastTangentToFirst(bool bMatchLastToFirst)
 			{
 				continue;
 			}
-			bFoundAnyTangents = true;
 			KeyHandles[0] = AllKeyHandles[0];
 			KeyHandles[1] = AllKeyHandles[AllKeyHandles.Num() - 1];
 			
@@ -1585,10 +1591,6 @@ void FCurveEditor::MatchLastTangentToFirst(bool bMatchLastToFirst)
 			CurveModel->SetKeyAttributes(KeyHandles, KeyAttributes);
 		}
 	}
-	if (!bFoundAnyTangents)
-	{
-		Transaction.Cancel();
-	}
 }
 
 bool FCurveEditor::CanPaste(const FString& TextToImport) const
@@ -1783,6 +1785,8 @@ void FCurveEditor::PasteKeys(TSet<FCurveModelID> CurveModelIDs, const bool bInOv
 
 void FCurveEditor::PasteKeys(UE::CurveEditor::FKeyPasteArgs InArgs)
 {
+	using namespace UE::CurveEditor;
+	
 	// Grab the text to paste from the clipboard
 	FString TextToImport;
 	FPlatformApplicationMisc::ClipboardPaste(TextToImport);
@@ -1847,10 +1851,11 @@ void FCurveEditor::PasteKeys(UE::CurveEditor::FKeyPasteArgs InArgs)
 	}
 	
 	const FScopedTransaction Transaction(LOCTEXT("PasteKeys", "Paste Keys")); // Selection and key change should be part of same undo / redo op.
-	const UE::CurveEditor::FScopedSelectionChange KeyChange(SharedThis(this));
+	const FScopedSelectionChange SelectionChange(SharedThis(this));
+	constexpr ECurveChangeFlags Flags = ECurveChangeFlags::AddKeys | ECurveChangeFlags::RemoveKeys | ECurveChangeFlags::KeyAttributes;
+	const FScopedKeyChange KeyChange(FCurvesSnapshotBuilder(SharedThis(this), TargetCurves, Flags));
 	Selection.Clear();
 
-
 	// Two simple cases, 1 to 1 and 1 to many.
 	TArray<TPair<UCurveEditorCopyableCurveKeys*, FCurveModelID>> CopyPairs;
 
@@ -1977,9 +1982,15 @@ void FCurveEditor::PasteKeys(UE::CurveEditor::FKeyPasteArgs InArgs)
 
 void FCurveEditor::DeleteSelection()
 {
-	const FScopedTransaction Transaction(LOCTEXT("DeleteKeys", "Delete Keys")); // Selection and key change should be part of same undo / redo op.
-	const UE::CurveEditor::FScopedSelectionChange SelectionChange(SharedThis(this));
 
+	using namespace UE::CurveEditor;
+	
+	const FScopedTransaction Transaction(LOCTEXT("DeleteKeys", "Delete Keys"));
+	const FScopedSelectionChange SelectionChange(SharedThis(this));
+	const FScopedKeyChange KeyChange(
+		FCurvesSnapshotBuilder(SharedThis(this), ECurveChangeFlags::RemoveKeys).TrackSelectedCurves()
+		);
+	
 	double CurrentTime = 0.0;
 	if (TSharedPtr<ITimeSliderController> TimeSliderController = WeakTimeSliderController.Pin())
 	{
@@ -1991,7 +2002,6 @@ void FCurveEditor::DeleteSelection()
 	{
 		if (FCurveModel* Curve = FindCurve(Pair.Key))
 		{
-			Curve->Modify();
 			Curve->RemoveKeys(Pair.Value.AsArray(), CurrentTime);
 		}
 	}
@@ -2001,8 +2011,12 @@ void FCurveEditor::DeleteSelection()
 
 void FCurveEditor::FlattenSelection()
 {
-	FScopedTransaction Transaction(LOCTEXT("FlattenTangents", "Flatten Tangents"));
-	bool bFoundAnyTangents = false;
+	using namespace UE::CurveEditor;
+	// FScopedKeyChange won't generate any transaction if no tangents are changed
+	const FScopedKeyChange KeyChange(
+		FCurvesSnapshotBuilder(SharedThis(this), CurveData, ECurveChangeFlags::KeyAttributes),
+		LOCTEXT("FlattenTangents", "Flatten Tangents")
+		);
 
 	TArray<FKeyHandle> KeyHandles;
 	TArray<FKeyAttributes> AllKeyPositions;
@@ -2068,31 +2082,27 @@ void FCurveEditor::FlattenSelection()
 
 			if (AllKeyPositions.Num() > 0)
 			{
-				Curve->Modify();
 				Curve->SetKeyAttributes(KeyHandles, AllKeyPositions);
 				if (KeyAttributesWeighted.Num() > 0)
 				{
 					Curve->SetKeyAttributes(KeyHandlesWeighted, KeyAttributesWeighted);
 				}
-				bFoundAnyTangents = true;
 			}
 		}
 	}
-
-	if (!bFoundAnyTangents)
-	{
-		Transaction.Cancel();
-	}
 }
 
 void FCurveEditor::StraightenSelection()
 {
-	FScopedTransaction Transaction(LOCTEXT("StraightenTangents", "Straighten Tangents"));
-	bool bFoundAnyTangents = false;
+	using namespace UE::CurveEditor;
+	// FScopedKeyChange won't generate any transaction if no tangents are changed
+	const FScopedKeyChange KeyChange(
+		FCurvesSnapshotBuilder(SharedThis(this), CurveData, ECurveChangeFlags::KeyAttributes),
+		LOCTEXT("StraightenTangents", "Straighten Tangents")
+		);
 
 	TArray<FKeyHandle> KeyHandles;
 	TArray<FKeyAttributes> AllKeyPositions;
-
 	for (const TTuple<FCurveModelID, FKeyHandleSet>& Pair : Selection.GetAll())
 	{
 		if (FCurveModel* Curve = FindCurve(Pair.Key))
@@ -2125,17 +2135,10 @@ void FCurveEditor::StraightenSelection()
 
 			if (AllKeyPositions.Num() > 0)
 			{
-				Curve->Modify();
 				Curve->SetKeyAttributes(KeyHandles, AllKeyPositions);
-				bFoundAnyTangents = true;
 			}
 		}
 	}
-
-	if (!bFoundAnyTangents)
-	{
-		Transaction.Cancel();
-	}
 }
 
 bool FCurveEditor::CanFlattenOrStraightenSelection() const
@@ -2145,14 +2148,17 @@ bool FCurveEditor::CanFlattenOrStraightenSelection() const
 
 void FCurveEditor::SmartSnapSelection()
 {
+	using namespace UE::CurveEditor;
 	const FScopedTransaction Transaction(LOCTEXT("SmartSnapKeys", "Smart Snap")); // Selection and key change should be part of same undo / redo op.
-	const UE::CurveEditor::FScopedSelectionChange SelectionChange(SharedThis(this));
+	const FScopedSelectionChange SelectionChange(SharedThis(this));
+	const FScopedKeyChange KeyChange(
+		FCurvesSnapshotBuilder(SharedThis(this), ECurveChangeFlags::MoveKeys | ECurveChangeFlags::RemoveKeys).TrackSelectedCurves()
+		);
 	
 	TMap<FCurveModelID, FKeyHandleSet> OutKeysToSelect;
-	UE::CurveEditor::EnumerateSmartSnappableKeys(*this, Selection.GetAll(), OutKeysToSelect,
-		[](const FCurveModelID&, FCurveModel& CurveModel, const UE::CurveEditor::FSmartSnapResult& SnapResult)
+	EnumerateSmartSnappableKeys(*this, Selection.GetAll(), OutKeysToSelect,
+		[](const FCurveModelID&, FCurveModel& CurveModel, const FSmartSnapResult& SnapResult)
 		{
-			CurveModel.Modify();
 			ApplySmartSnap(CurveModel, SnapResult);
 		});
 
@@ -2355,7 +2361,12 @@ void FCurveEditor::ApplyBufferedCurveToTarget(const IBufferedCurveModel* Buffere
 
 bool FCurveEditor::ApplyBufferedCurves(const TSet<FCurveModelID>& InCurvesToApplyTo, const bool bSwapBufferCurves)
 {
-	FScopedTransaction Transaction(bSwapBufferCurves ? LOCTEXT("SwapBufferedCurves", "Swap Buffered Curves") : LOCTEXT("ApplyBufferedCurves", "Apply Buffered Curves"));
+	using namespace UE::CurveEditor;
+	// Fyi, FScopedKeyChange won't make any entry in the Undo/Redo buffer if nothing was changed.
+	const FScopedKeyChange KeyChange(
+		FCurvesSnapshotBuilder(SharedThis(this), InCurvesToApplyTo),
+		bSwapBufferCurves ? LOCTEXT("SwapBufferedCurves", "Swap Buffered Curves") : LOCTEXT("ApplyBufferedCurves", "Apply Buffered Curves")
+		);
 
 	// Each curve can specify an "Intention" name. This gives a little bit of context about how the curve is intended to be used,
 	// without locking anyone into a specific set of intentions. When you go to apply the buffered curves, for each curve that you
@@ -2369,7 +2380,6 @@ bool FCurveEditor::ApplyBufferedCurves(const TSet<FCurveModelID>& InCurvesToAppl
 	bool bFoundAnyMatchedIntent = false;
 
 	TMap<FString, int32> IntentMatchIndexes;
-
 	for (const FCurveModelID& CurveModelID : InCurvesToApplyTo)
 	{
 		FCurveModel* TargetCurve = FindCurve(CurveModelID);
@@ -2384,8 +2394,6 @@ bool FCurveEditor::ApplyBufferedCurves(const TSet<FCurveModelID>& InCurvesToAppl
 			continue;
 		}
 
-		TargetCurve->Modify();
-
 		// In an attempt to support buffering multiple curves with the same intention, we'll try to match them up in pairs. This means
 		// for the first curve that we're trying to apply to, if the intention is "Location.X" we will search the buffered curves for a
 		// "Location.X". Upon finding one, we store the index that it was found at, so the next time we try to find a curve with the same
@@ -2538,8 +2546,6 @@ bool FCurveEditor::ApplyBufferedCurves(const TSet<FCurveModelID>& InCurvesToAppl
 		FSlateNotificationManager::Get().AddNotification(FailInfo);
 	}
 
-	// No need to make a entry in the Undo/Redo buffer if it didn't apply anything.
-	Transaction.Cancel();
 	return false;
 }
 
diff --git a/Engine/Source/Editor/CurveEditor/Private/DragOperations/CurveEditorDragOperation_Marquee.cpp b/Engine/Source/Editor/CurveEditor/Private/DragOperations/CurveEditorDragOperation_Marquee.cpp
--- a/Engine/Source/Editor/CurveEditor/Private/DragOperations/CurveEditorDragOperation_Marquee.cpp
+++ b/Engine/Source/Editor/CurveEditor/Private/DragOperations/CurveEditorDragOperation_Marquee.cpp
@@ -155,13 +155,13 @@ FCurveEditorDragOperation_Marquee::FCurveEditorDragOperation_Marquee(FCurveEdito
 {}
 
 FCurveEditorDragOperation_Marquee::FCurveEditorDragOperation_Marquee(FCurveEditor* InCurveEditor, SCurveEditorView* InLockedToView)
-	: SelectionChange(InCurveEditor->AsShared())
-	, CurveEditor(InCurveEditor)
+	: CurveEditor(InCurveEditor)
 	, LockedToView(InLockedToView)
 {}
 
 void FCurveEditorDragOperation_Marquee::OnBeginDrag(FVector2D InitialPosition, FVector2D CurrentPosition, const FPointerEvent& MouseEvent)
 {
+	SelectionChange.Emplace(CurveEditor->AsShared());
 	RealInitialPosition = CurrentPosition;
 
 	Marquee = FSlateRect(
@@ -246,7 +246,8 @@ void FCurveEditorDragOperation_Marquee::OnEndDrag(FVector2D InitialPosition, FVe
 			CurveEditor->Selection.Add(Point);
 		}
 	}
-	
+
+	SelectionChange.Reset();
 }
 
 void FCurveEditorDragOperation_Marquee::OnPaint(const FGeometry& AllottedGeometry, FSlateWindowElementList& OutDrawElements, int32 PaintOnLayerId)
diff --git a/Engine/Source/Editor/CurveEditor/Private/DragOperations/CurveEditorDragOperation_Marquee.h b/Engine/Source/Editor/CurveEditor/Private/DragOperations/CurveEditorDragOperation_Marquee.h
--- a/Engine/Source/Editor/CurveEditor/Private/DragOperations/CurveEditorDragOperation_Marquee.h
+++ b/Engine/Source/Editor/CurveEditor/Private/DragOperations/CurveEditorDragOperation_Marquee.h
@@ -20,7 +20,7 @@ class FCurveEditorDragOperation_Marquee : public ICurveEditorDragOperation
 public:
 
 	FCurveEditorDragOperation_Marquee(FCurveEditor* InCurveEditor);
-	FCurveEditorDragOperation_Marquee(FCurveEditor* InCurveEditor, SCurveEditorView*  InLockedToView);
+	FCurveEditorDragOperation_Marquee(FCurveEditor* InCurveEditor, SCurveEditorView* InLockedToView);
 
 	virtual void OnBeginDrag(FVector2D InitialPosition, FVector2D CurrentPosition, const FPointerEvent& MouseEvent) override;
 	virtual void OnDrag(FVector2D InitialPosition, FVector2D CurrentPosition, const FPointerEvent& MouseEvent) override;
@@ -29,8 +29,8 @@ public:
 
 private:
 
-	/** Makes this change undo / redo able. */
-	const UE::CurveEditor::FScopedSelectionChange SelectionChange;
+	/** Makes this change undo / redo able. Set for the duration of the drag. */
+	TOptional<UE::CurveEditor::FScopedSelectionChange> SelectionChange;
 
 	/** The current marquee rectangle */
 	FSlateRect Marquee;
diff --git a/Engine/Source/Editor/CurveEditor/Private/DragOperations/CurveEditorDragOperation_MoveKeys.cpp b/Engine/Source/Editor/CurveEditor/Private/DragOperations/CurveEditorDragOperation_MoveKeys.cpp
--- a/Engine/Source/Editor/CurveEditor/Private/DragOperations/CurveEditorDragOperation_MoveKeys.cpp
+++ b/Engine/Source/Editor/CurveEditor/Private/DragOperations/CurveEditorDragOperation_MoveKeys.cpp
@@ -17,6 +17,7 @@
 #include "Misc/AssertionMacros.h"
 #include "SCurveEditorView.h"
 #include "ScopedTransaction.h"
+#include "Modification/Utils/ScopedKeyChange.h"
 #include "Templates/Tuple.h"
 #include "Templates/UniquePtr.h"
 #include "Templates/UnrealTemplate.h"
@@ -32,8 +33,18 @@ void FCurveEditorDragOperation_MoveKeys::OnInitialize(FCurveEditor* InCurveEdito
 
 void FCurveEditorDragOperation_MoveKeys::OnBeginDrag(FVector2D InitialPosition, FVector2D CurrentPosition, const FPointerEvent& MouseEvent)
 {
-	int32 NumKeys = CurveEditor->GetSelection().Count();
-	Transaction = MakeUnique<FScopedTransaction>(FText::Format(NSLOCTEXT("CurveEditor", "MoveKeysFormat", "Move {0}|plural(one=Key, other=Keys)"), NumKeys));
+	using namespace UE::CurveEditor;
+
+	// In OnBeginDrag, we should already be in a transaction. Setting "Transaction" will extend that transaction until the user lets go.
+	const int32 NumKeys = CurveEditor->GetSelection().Count();
+	Transaction = MakeUnique<FScopedTransaction>(
+		FText::Format(NSLOCTEXT("CurveEditor", "MoveKeysFormat", "Move {0}|plural(one=Key, other=Keys)"), NumKeys)
+		);
+	
+	TSharedRef<FCurveEditor> CurveEditorPin = CurveEditor->AsShared();
+	ScopedKeyChange.Emplace(
+		FCurvesSnapshotBuilder(CurveEditorPin, ECurveChangeFlags::MoveKeys).TrackSelectedKeys()
+		);
 
 	KeysByCurve.Reset();
 	CurveEditor->SuppressBoundTransformUpdates(true);
@@ -46,8 +57,6 @@ void FCurveEditorDragOperation_MoveKeys::OnBeginDrag(FVector2D InitialPosition,
 
 		if (ensureAlways(Curve))
 		{
-			Curve->Modify();
-
 			TArrayView<const FKeyHandle> Handles = Pair.Value.AsArray();
 
 			FKeyData& KeyData = KeysByCurve.Emplace_GetRef(CurveID);
@@ -102,6 +111,10 @@ void FCurveEditorDragOperation_MoveKeys::OnCancelDrag()
 	}
 
 	CurveEditor->SuppressBoundTransformUpdates(false);
+	
+	ScopedKeyChange->Cancel();
+	ScopedKeyChange.Reset();
+	Transaction.Reset();
 }
 
 void FCurveEditorDragOperation_MoveKeys::OnEndDrag(FVector2D InitialPosition, FVector2D CurrentPosition, const FPointerEvent& MouseEvent)
@@ -117,6 +130,9 @@ void FCurveEditorDragOperation_MoveKeys::OnEndDrag(FVector2D InitialPosition, FV
 	}
 
 	CurveEditor->SuppressBoundTransformUpdates(false);
+	
+	ScopedKeyChange.Reset();
+	Transaction.Reset();
 }
 
 FCurveEditorDragOperation_MoveKeys::FAccumulatedMouseMovement& FCurveEditorDragOperation_MoveKeys::GetOrAddAccumulatedMouseMovement(
diff --git a/Engine/Source/Editor/CurveEditor/Private/DragOperations/CurveEditorDragOperation_MoveKeys.h b/Engine/Source/Editor/CurveEditor/Private/DragOperations/CurveEditorDragOperation_MoveKeys.h
--- a/Engine/Source/Editor/CurveEditor/Private/DragOperations/CurveEditorDragOperation_MoveKeys.h
+++ b/Engine/Source/Editor/CurveEditor/Private/DragOperations/CurveEditorDragOperation_MoveKeys.h
@@ -10,6 +10,7 @@
 #include "ICurveEditorDragOperation.h"
 #include "Math/Vector2D.h"
 #include "Misc/Optional.h"
+#include "Modification/Utils/ScopedKeyChange.h"
 
 class FCurveEditor;
 struct FPointerEvent;
@@ -34,7 +35,9 @@ private:
 	TOptional<FCurvePointHandle> CardinalPoint;
 
 	FVector2D LastMousePosition;
-private:
+
+	/** Transacts the moved keys. */
+	TOptional<UE::CurveEditor::FScopedKeyChange> ScopedKeyChange;
 
 	struct FKeyData
 	{
diff --git a/Engine/Source/Editor/CurveEditor/Private/DragOperations/CurveEditorDragOperation_Tangent.cpp b/Engine/Source/Editor/CurveEditor/Private/DragOperations/CurveEditorDragOperation_Tangent.cpp
--- a/Engine/Source/Editor/CurveEditor/Private/DragOperations/CurveEditorDragOperation_Tangent.cpp
+++ b/Engine/Source/Editor/CurveEditor/Private/DragOperations/CurveEditorDragOperation_Tangent.cpp
@@ -229,10 +229,18 @@ void FCurveEditorDragOperation_Tangent::OnInitialize(FCurveEditor* InCurveEditor
 
 void FCurveEditorDragOperation_Tangent::OnBeginDrag(FVector2D InitialPosition, FVector2D CurrentPosition, const FPointerEvent& MouseEvent)
 {
-	int32 NumKeys = CurveEditor->GetSelection().Count();
+	using namespace UE::CurveEditor;
 
-	FText Description = FText::Format(NSLOCTEXT("CurveEditor", "DragTangentsFormat", "Drag {0}|plural(one=Tangent, other=Tangents)"), NumKeys);
-	Transaction = MakeUnique<FScopedTransaction>(Description);
+	// In OnBeginDrag, we should already be in a transaction. Setting "Transaction" will extend that transaction until the user lets go.
+	const int32 NumKeys = CurveEditor->GetSelection().Count();
+	Transaction = MakeUnique<FScopedTransaction>(
+		FText::Format(NSLOCTEXT("CurveEditor", "DragTangentsFormat", "Drag {0}|plural(one=Tangent, other=Tangents)"), NumKeys)
+		);
+	
+	TSharedRef<FCurveEditor> CurveEditorPin = CurveEditor->AsShared();
+	ScopedKeyChange.Emplace(
+		FCurvesSnapshotBuilder(CurveEditorPin, ECurveChangeFlags::KeyAttributes).TrackSelectedCurves()
+		);
 	CurveEditor->SuppressBoundTransformUpdates(true);
 
 	KeysByCurve.Reset();
@@ -243,8 +251,6 @@ void FCurveEditorDragOperation_Tangent::OnBeginDrag(FVector2D InitialPosition, F
 
 		if (ensureAlways(Curve))
 		{
-			Curve->Modify();
-
 			TArrayView<const FKeyHandle> Handles = Pair.Value.AsArray();
 
 			FKeyData& KeyData = KeysByCurve.Emplace_GetRef(CurveID);
@@ -320,6 +326,10 @@ void FCurveEditorDragOperation_Tangent::OnCancelDrag()
 	}
 
 	CurveEditor->SuppressBoundTransformUpdates(false);
+	
+	ScopedKeyChange->Cancel();
+	ScopedKeyChange.Reset();
+	Transaction.Reset();
 }
 
 void FCurveEditorDragOperation_Tangent::OnEndDrag(FVector2D InitialPosition, FVector2D CurrentPosition, const FPointerEvent& MouseEvent)
@@ -339,4 +349,6 @@ void FCurveEditorDragOperation_Tangent::OnEndDrag(FVector2D InitialPosition, FVe
 	}
 	
 	CurveEditor->SuppressBoundTransformUpdates(false);
+	ScopedKeyChange.Reset();
+	Transaction.Reset();
 }
diff --git a/Engine/Source/Editor/CurveEditor/Private/DragOperations/CurveEditorDragOperation_Tangent.h b/Engine/Source/Editor/CurveEditor/Private/DragOperations/CurveEditorDragOperation_Tangent.h
--- a/Engine/Source/Editor/CurveEditor/Private/DragOperations/CurveEditorDragOperation_Tangent.h
+++ b/Engine/Source/Editor/CurveEditor/Private/DragOperations/CurveEditorDragOperation_Tangent.h
@@ -10,6 +10,7 @@
 #include "Math/Vector2D.h"
 #include "Misc/EnumClassFlags.h"
 #include "Misc/Optional.h"
+#include "Modification/Utils/ScopedKeyChange.h"
 
 class FCurveEditor;
 struct FCurveEditorScreenSpace;
@@ -88,7 +89,7 @@ namespace CurveEditorDragOperation
 		 * an ugly pop in the curve and it lets the Arrive tangents become Leave tangents which defeats the point.
 		 */
 		void MendNearlyZeroXTangentOffset(const bool bArriveTangent, FVector2D& InOutTangentOffset) const;
-
+		
 		/** Holds the last tanget offset. Useful when editing the tangent or weight exclusively. */
 		TOptional<FVector2D> LastTangentOffset;
 
@@ -99,7 +100,7 @@ namespace CurveEditorDragOperation
 		const FKeyAttributes InitialKeyAttributes;
 
 		/** The point type of the key handle for which the tangent is solved */
-		const ECurvePointType PointType; 
+		const ECurvePointType PointType;
 	};
 }
 
@@ -137,4 +138,7 @@ private:
 
 	/** Key dragging data stored per-curve */
 	TArray<FKeyData> KeysByCurve;
+	
+	/** Transacts the keys' tangents. */
+	TOptional<UE::CurveEditor::FScopedKeyChange> ScopedKeyChange;
 };
\ No newline at end of file
diff --git a/Engine/Source/Editor/CurveEditor/Private/Filters/CurveEditorBakeFilter.cpp b/Engine/Source/Editor/CurveEditor/Private/Filters/CurveEditorBakeFilter.cpp
--- a/Engine/Source/Editor/CurveEditor/Private/Filters/CurveEditorBakeFilter.cpp
+++ b/Engine/Source/Editor/CurveEditor/Private/Filters/CurveEditorBakeFilter.cpp
@@ -203,8 +203,6 @@ void UCurveEditorBakeFilter::ApplyFilter_Impl(TSharedRef<FCurveEditor> InCurveEd
 			}
 		}
 
-		Curve->Modify();
-
 		// We want to store the newly added ones so we can add them to the users selection to mimic
 		// their selection pre-baking.
 		TArray<TOptional<FKeyHandle>> NewKeyHandles;
diff --git a/Engine/Source/Editor/CurveEditor/Private/Filters/CurveEditorEulerFilter.cpp b/Engine/Source/Editor/CurveEditor/Private/Filters/CurveEditorEulerFilter.cpp
--- a/Engine/Source/Editor/CurveEditor/Private/Filters/CurveEditorEulerFilter.cpp
+++ b/Engine/Source/Editor/CurveEditor/Private/Filters/CurveEditorEulerFilter.cpp
@@ -85,7 +85,6 @@ void UCurveEditorEulerFilter::ApplyFilter_Impl(TSharedRef<FCurveEditor> InCurveE
 				KeyHandlesToModify.Add(KeyHandles[KeyIndex + 1]);
 			}
 
-			Curve->Modify();
 			Curve->SetKeyPositions(KeyHandlesToModify, NewKeyPositions);
 		}
 	}
diff --git a/Engine/Source/Editor/CurveEditor/Private/Filters/CurveEditorReduceFilter.cpp b/Engine/Source/Editor/CurveEditor/Private/Filters/CurveEditorReduceFilter.cpp
--- a/Engine/Source/Editor/CurveEditor/Private/Filters/CurveEditorReduceFilter.cpp
+++ b/Engine/Source/Editor/CurveEditor/Private/Filters/CurveEditorReduceFilter.cpp
@@ -180,7 +180,6 @@ void UCurveEditorReduceFilter::ApplyFilter_Impl(TSharedRef<FCurveEditor> InCurve
 					OutHandleSet.Remove(KeyHandles[TestIndex], ECurvePointType::Key);
 				}
 			}
-			Curve->Modify();
 			Curve->RemoveKeys(KeysToRemove,0.0);
 		}
 	}
diff --git a/Engine/Source/Editor/CurveEditor/Private/Filters/CurveEditorSmartReduceFilter.cpp b/Engine/Source/Editor/CurveEditor/Private/Filters/CurveEditorSmartReduceFilter.cpp
--- a/Engine/Source/Editor/CurveEditor/Private/Filters/CurveEditorSmartReduceFilter.cpp
+++ b/Engine/Source/Editor/CurveEditor/Private/Filters/CurveEditorSmartReduceFilter.cpp
@@ -211,7 +211,6 @@ void UCurveEditorSmartReduceFilter::SmartReduce(FCurveModel* Curve, const FSmart
 
 	TArray<double> Velocities;
 	Difference(Cache.Positions, Velocities);
-	Curve->Modify();
 	Curve->RemoveKeys(KeyHandles, 0.0);
 	FKeyAttributes NewDefaultAttribute;
 	NewDefaultAttribute.SetInterpMode(ERichCurveInterpMode::RCIM_Cubic);
diff --git a/Engine/Source/Editor/CurveEditor/Private/Filters/CurveEditorSmartSnapFilter.cpp b/Engine/Source/Editor/CurveEditor/Private/Filters/CurveEditorSmartSnapFilter.cpp
--- a/Engine/Source/Editor/CurveEditor/Private/Filters/CurveEditorSmartSnapFilter.cpp
+++ b/Engine/Source/Editor/CurveEditor/Private/Filters/CurveEditorSmartSnapFilter.cpp
@@ -14,7 +14,6 @@ void UCurveEditorSmartSnapFilter::ApplyFilter_Impl(
 	UE::CurveEditor::EnumerateSmartSnappableKeys(*InCurveEditor, InKeysToOperateOn, OutKeysToSelect,
 		[&OutKeysToSelect](const FCurveModelID& CurveModelId, FCurveModel& CurveModel, const UE::CurveEditor::FSmartSnapResult& SnapResult)
 		{
-			CurveModel.Modify();
 			ApplySmartSnap(CurveModel, SnapResult);
 		});
 }
diff --git a/Engine/Source/Editor/CurveEditor/Private/Filters/FilterUtils.cpp b/Engine/Source/Editor/CurveEditor/Private/Filters/FilterUtils.cpp
--- a/Engine/Source/Editor/CurveEditor/Private/Filters/FilterUtils.cpp
+++ b/Engine/Source/Editor/CurveEditor/Private/Filters/FilterUtils.cpp
@@ -6,6 +6,7 @@
 #include "Modification/Utils/ScopedSelectionChange.h"
 
 #include "Containers/Map.h"
+#include "Modification/Utils/ScopedKeyChange.h"
 #include "Templates/SharedPointer.h"
 
 namespace UE::CurveEditor::FilterUtils
@@ -21,8 +22,13 @@ namespace UE::CurveEditor::FilterUtils
 		const TMap<FCurveModelID, FKeyHandleSet>& InSelectedKeys
 		)
 	{
+		using namespace UE::CurveEditor;
+		
 		TMap<FCurveModelID, FKeyHandleSet> OutKeysToSelect;
-		InFilter.ApplyFilter(InCurveEditor, InSelectedKeys, OutKeysToSelect);
+		{
+			const FScopedKeyChange KeyChange(FCurvesSnapshotBuilder(InCurveEditor, InSelectedKeys));
+			InFilter.ApplyFilter(InCurveEditor, InSelectedKeys, OutKeysToSelect);
+		}
 		
 		const FScopedSelectionChange Transaction(InCurveEditor);
 		// Clear their selection and then set it to the keys the filter thinks you should have selected.
diff --git a/Engine/Source/Editor/CurveEditor/Private/SCurveEditorPanel.cpp b/Engine/Source/Editor/CurveEditor/Private/SCurveEditorPanel.cpp
--- a/Engine/Source/Editor/CurveEditor/Private/SCurveEditorPanel.cpp
+++ b/Engine/Source/Editor/CurveEditor/Private/SCurveEditorPanel.cpp
@@ -67,6 +67,7 @@
 #include "Filters/CurveEditorEulerFilter.h"
 #include "Filters/PromotedFilterContainer.h"
 #include "Misc/ResizeParamUtils.h"
+#include "Modification/Utils/ScopedKeyChange.h"
 #include "Modification/Utils/ScopedSelectionChange.h"
 #include "Styling/AppStyle.h"
 #include "Styling/CoreStyle.h"
@@ -1045,13 +1046,16 @@ EVisibility SCurveEditorPanel::GetSplitterVisibility() const
 
 void SCurveEditorPanel::SetKeyAttributes(FKeyAttributes KeyAttributes, FText Description)
 {
-	FScopedTransaction Transaction(Description);
-
+	using namespace UE::CurveEditor;
+	const FScopedKeyChange KeyChange(
+		FCurvesSnapshotBuilder(GetCurveEditor(), ECurveChangeFlags::KeyAttributes).TrackSelectedCurves(),
+		Description
+		);
+	
 	for (const TTuple<FCurveModelID, FKeyHandleSet>& Pair : CurveEditor->Selection.GetAll())
 	{
 		if (FCurveModel* Curve = CurveEditor->FindCurve(Pair.Key))
 		{
-			Curve->Modify();
 			Curve->SetKeyAttributes(Pair.Value.AsArray(), KeyAttributes);
 		}
 	}
@@ -1059,13 +1063,16 @@ void SCurveEditorPanel::SetKeyAttributes(FKeyAttributes KeyAttributes, FText Des
 
 void SCurveEditorPanel::SetCurveAttributes(FCurveAttributes CurveAttributes, FText Description)
 {
-	FScopedTransaction Transaction(Description);
+	using namespace UE::CurveEditor;
+	const FScopedKeyChange KeyChange(
+		FCurvesSnapshotBuilder(CurveEditor, ECurveChangeFlags::CurveAttributes).TrackSelectedCurves(),
+		Description
+		);
 
 	for (const TTuple<FCurveModelID, FKeyHandleSet>& Pair : CurveEditor->Selection.GetAll())
 	{
 		if (FCurveModel* Curve = CurveEditor->FindCurve(Pair.Key))
 		{
-			Curve->Modify();
 			Curve->SetCurveAttributes(CurveAttributes);
 		}
 	}
@@ -1073,16 +1080,17 @@ void SCurveEditorPanel::SetCurveAttributes(FCurveAttributes CurveAttributes, FTe
 
 void SCurveEditorPanel::ToggleWeightedTangents()
 {
-	FScopedTransaction Transaction(LOCTEXT("ToggleWeightedTangents_Transaction", "Toggle Weighted Tangents"));
-
-	TMap<FCurveModelID, TArray<FKeyAttributes>> KeyAttributesPerCurve;
-
+	using namespace UE::CurveEditor;
 	const TMap<FCurveModelID, FKeyHandleSet>& Selection = CurveEditor->GetSelection().GetAll();
+	const FScopedKeyChange KeyChange(
+		FCurvesSnapshotBuilder(CurveEditor, Selection, ECurveChangeFlags::KeyAttributes),
+		LOCTEXT("ToggleWeightedTangents_Transaction", "Toggle Weighted Tangents")
+		);
 
 	// Disable weights unless we find something that doesn't have weights, then add them
 	FKeyAttributes KeyAttributesToAssign = FKeyAttributes().SetTangentWeightMode(RCTWM_WeightedNone);
-
 	// Gather current key attributes
+	TMap<FCurveModelID, TArray<FKeyAttributes>> KeyAttributesPerCurve;
 	for (const TTuple<FCurveModelID, FKeyHandleSet>& Pair : Selection)
 	{
 		FCurveModel* Curve = CurveEditor->FindCurve(Pair.Key);
@@ -1119,7 +1127,6 @@ void SCurveEditorPanel::ToggleWeightedTangents()
 			}
 
 			TArrayView<const FKeyHandle> KeyHandles = Selection.FindChecked(Pair.Key).AsArray();
-			Curve->Modify();
 			Curve->SetKeyAttributes(KeyHandles, Pair.Value);
 		}
 	}
diff --git a/Engine/Source/Editor/CurveEditor/Private/Views/SInteractiveCurveEditorView.cpp b/Engine/Source/Editor/CurveEditor/Private/Views/SInteractiveCurveEditorView.cpp
--- a/Engine/Source/Editor/CurveEditor/Private/Views/SInteractiveCurveEditorView.cpp
+++ b/Engine/Source/Editor/CurveEditor/Private/Views/SInteractiveCurveEditorView.cpp
@@ -1593,6 +1593,8 @@ void SInteractiveCurveEditorView::HandleDirectKeySelectionByMouse(TSharedPtr<FCu
 
 FReply SInteractiveCurveEditorView::OnMouseButtonDown(const FGeometry& MyGeometry, const FPointerEvent& MouseEvent)
 {
+	using namespace UE::CurveEditor;
+	
 	TSharedPtr<FCurveEditor> CurveEditor = WeakCurveEditor.Pin();
 	TSharedPtr<SCurveEditorPanel> EditorPanel = CurveEditor.IsValid() ? CurveEditor->GetPanel() : nullptr;
 	if (!CurveEditor || !EditorPanel)
@@ -1626,7 +1628,10 @@ FReply SInteractiveCurveEditorView::OnMouseButtonDown(const FGeometry& MyGeometr
 			if (CurveToAddTo && !CurveToAddTo->IsReadOnly() && !bDraggingKeys && !bKeyAlreadyExists)
 			{
 				FScopedTransaction Transaction(LOCTEXT("InsertKey", "Insert Key")); // Selection and key change should be part of same undo / redo op.
-				UE::CurveEditor::FScopedSelectionChange SelectionChange(CurveEditor);
+				const FScopedSelectionChange SelectionChange(CurveEditor);
+				const FScopedKeyChange KeyChange(
+					FCurvesSnapshotBuilder(CurveEditor, *HoveredCurve, ECurveChangeFlags::KeyAttributes | ECurveChangeFlags::AddKeys)
+					);
 
 				FCurveEditorScreenSpace CurveSpace = GetCurveSpace(HoveredCurve.GetValue());
 				double MouseTime = CurveSpace.ScreenToSeconds(MousePixel.X);
@@ -1667,8 +1672,6 @@ FReply SInteractiveCurveEditorView::OnMouseButtonDown(const FGeometry& MyGeometr
 					MouseValue = CurveOutputMin;
 				}
 
-				CurveToAddTo->Modify();
-
 				// Add a key on this curve
 				TOptional<FKeyHandle> NewKey = CurveToAddTo->AddKey(FKeyPosition(MouseTime, MouseValue), KeyAttributes);
 				if (NewKey.IsSet())
@@ -2113,6 +2116,8 @@ void SInteractiveCurveEditorView::AddKeyAtMousePosition(TSet<FCurveModelID> ForC
 
 void SInteractiveCurveEditorView::AddKeyAtTime(const TSet<FCurveModelID>& ToCurves, double InTime)
 {
+	using namespace UE::CurveEditor;
+	
 	TSharedPtr<FCurveEditor> CurveEditor = WeakCurveEditor.Pin();
 	if (!CurveEditor.IsValid())
 	{
@@ -2120,7 +2125,11 @@ void SInteractiveCurveEditorView::AddKeyAtTime(const TSet<FCurveModelID>& ToCurv
 	}
 
 	FScopedTransaction Transaction(LOCTEXT("AddKeyAtTime", "Add Key")); // Selection and key change should be part of same undo / redo op.
-	UE::CurveEditor::FScopedSelectionChange SelectionChange(WeakCurveEditor);
+	FScopedSelectionChange SelectionChange(WeakCurveEditor);
+	FScopedKeyChange KeyChange(
+		FCurvesSnapshotBuilder(WeakCurveEditor, ToCurves, ECurveChangeFlags::AddKeys)
+		);
+	
 	bool bAddedKey = false;
 
 	// Clear the selection set as we will be selecting all the new keys created.
@@ -2146,7 +2155,6 @@ void SInteractiveCurveEditorView::AddKeyAtTime(const TSet<FCurveModelID>& ToCurv
 		double CurveValue = 0.0;
 		if (CurveModel->Evaluate(EvalTime, CurveValue))
 		{
-			CurveModel->Modify();
 			CurveValue = SnapMetrics.SnapOutput(CurveValue);
 
 			// Curve Models allow us to create new keys ontop of existing keys which works, but causes some user confusion
diff --git a/Engine/Source/Editor/CurveEditor/Public/Misc/CurveEditorUtils.h b/Engine/Source/Editor/CurveEditor/Public/Misc/CurveEditorUtils.h
deleted file mode 100644
diff --git a/Engine/Source/Editor/Persona/Private/SAnimSequenceCurveEditor.cpp b/Engine/Source/Editor/Persona/Private/SAnimSequenceCurveEditor.cpp
--- a/Engine/Source/Editor/Persona/Private/SAnimSequenceCurveEditor.cpp
+++ b/Engine/Source/Editor/Persona/Private/SAnimSequenceCurveEditor.cpp
@@ -112,8 +112,7 @@ void FRichCurveEditorModelNamed::SetKeyAttributes(TArrayView<const FKeyHandle> I
 }
 
 void FRichCurveEditorModelNamed::SetCurveAttributes(const FCurveAttributes& InCurveAttributes)
-{
-	Modify();		
+{	
 	IAnimationDataController& Controller = AnimSequence->GetController();
 	Controller.SetCurveAttributes(CurveId, InCurveAttributes);
 }
