Curve Editor: Replaces obvious usages GetKeys with GetAllKeys.

diff --git a/Engine/Plugins/Animation/TweeningUtils/Source/TweeningUtilsEditor/Private/Math/ContiguousKeyMapping.cpp b/Engine/Plugins/Animation/TweeningUtils/Source/TweeningUtilsEditor/Private/Math/ContiguousKeyMapping.cpp
--- a/Engine/Plugins/Animation/TweeningUtils/Source/TweeningUtilsEditor/Private/Math/ContiguousKeyMapping.cpp
+++ b/Engine/Plugins/Animation/TweeningUtils/Source/TweeningUtilsEditor/Private/Math/ContiguousKeyMapping.cpp
@@ -45,11 +45,7 @@ static void AppendKeyArray(
 {
 	if (const FCurveModel* Curve = InCurveEditor.FindCurve(InCurveId))
 	{
-		TArray<FKeyHandle> AllKeyHandles;
-		Curve->GetKeys(
-			TNumericLimits<double>::Lowest(), TNumericLimits<double>::Max(), TNumericLimits<double>::Lowest(), TNumericLimits<double>::Max(),
-			AllKeyHandles
-			);
+		const TArray<FKeyHandle> AllKeyHandles = Curve->GetAllKeys();
 		
 		TArray<FKeyPosition> AllKeyPositions;
 		AllKeyPositions.SetNum(AllKeyHandles.Num());
diff --git a/Engine/Plugins/Editor/CurveEditorTools/Source/CurveEditorTools/Private/Tools/CurveEditorRetimeTool.cpp b/Engine/Plugins/Editor/CurveEditorTools/Source/CurveEditorTools/Private/Tools/CurveEditorRetimeTool.cpp
--- a/Engine/Plugins/Editor/CurveEditorTools/Source/CurveEditorTools/Private/Tools/CurveEditorRetimeTool.cpp
+++ b/Engine/Plugins/Editor/CurveEditorTools/Source/CurveEditorTools/Private/Tools/CurveEditorRetimeTool.cpp
@@ -507,7 +507,7 @@ void FCurveEditorRetimeTool::OnDragStart()
 			FPreDragChannelData& ChannelData = ResizeData.CurveChannels[ResizeData.CurveChannels.Emplace(CurveID)];
 
 			// Get all Key Handles
-			CurveModel->GetKeys(TNumericLimits<double>::Lowest(), TNumericLimits<double>::Max(), TNumericLimits<double>::Lowest(), TNumericLimits<double>::Max(), ChannelData.Handles);
+			CurveModel->GetAllKeys(ChannelData.Handles);
 
 			// Get all Key Positions for those Handles.
 			ChannelData.FrameNumbers.SetNumUninitialized(ChannelData.Handles.Num());
diff --git a/Engine/Plugins/FX/Niagara/Source/NiagaraEditorWidgets/Private/DetailCustomizations/NiagaraDataInterfaceCurveDetails.cpp b/Engine/Plugins/FX/Niagara/Source/NiagaraEditorWidgets/Private/DetailCustomizations/NiagaraDataInterfaceCurveDetails.cpp
--- a/Engine/Plugins/FX/Niagara/Source/NiagaraEditorWidgets/Private/DetailCustomizations/NiagaraDataInterfaceCurveDetails.cpp
+++ b/Engine/Plugins/FX/Niagara/Source/NiagaraEditorWidgets/Private/DetailCustomizations/NiagaraDataInterfaceCurveDetails.cpp
@@ -398,8 +398,7 @@ private:
 
 	void GetSortedKeyHandlessAndPositionsForModel(FCurveModel& InCurveModel, TArray<FKeyHandlePositionPair>& OutSortedKeyHandlesAndPositions)
 	{
-		TArray<FKeyHandle> KeyHandles;
-		InCurveModel.GetKeys(TNumericLimits<double>::Lowest(), TNumericLimits<double>::Max(), TNumericLimits<double>::Lowest(), TNumericLimits<double>::Max(), KeyHandles);
+		const TArray<FKeyHandle> KeyHandles = InCurveModel.GetAllKeys();
 
 		TArray<FKeyPosition> KeyPositions;
 		KeyPositions.AddDefaulted(KeyHandles.Num());
@@ -681,8 +680,7 @@ private:
 			else if (CurveModelForAdd->GetNumKeys() == 1)
 			{
 				// If there's a single key, add the new key at the same value, but time + 1.
-				TArray<FKeyHandle> KeyHandles;
-				CurveModelForAdd->GetKeys(TNumericLimits<double>::Lowest(), TNumericLimits<double>::Max(), TNumericLimits<double>::Lowest(), TNumericLimits<double>::Max(), KeyHandles);
+				const TArray<FKeyHandle> KeyHandles = CurveModelForAdd->GetAllKeys();
 				
 				TArray<FKeyPosition> KeyPositions;
 				KeyPositions.AddDefaulted();
@@ -1609,8 +1607,7 @@ FReply FNiagaraDataInterfaceCurveDetailsBase::CurveTemplateSelected(TWeakObjectP
 				FCurveModel* CurveModel = CurveEditor->GetCurves()[CurveModelId].Get();
 				if (CurveModel != nullptr)
 				{
-					TArray<FKeyHandle> KeyHandles;
-					CurveModel->GetKeys(TNumericLimits<double>::Lowest(), TNumericLimits<double>::Max(), TNumericLimits<double>::Lowest(), TNumericLimits<double>::Max(), KeyHandles);
+					const TArray<FKeyHandle> KeyHandles = CurveModel->GetAllKeys();
 					CurveModel->RemoveKeys(KeyHandles,0.0);
 
 					const FRichCurve& FloatCurve = FloatCurveAsset->FloatCurve;
diff --git a/Engine/Source/Editor/CurveEditor/Private/CurveEditor.cpp b/Engine/Source/Editor/CurveEditor/Private/CurveEditor.cpp
--- a/Engine/Source/Editor/CurveEditor/Private/CurveEditor.cpp
+++ b/Engine/Source/Editor/CurveEditor/Private/CurveEditor.cpp
@@ -1053,8 +1053,7 @@ void FCurveEditor::SelectAllKeys()
 	{
 		if (FCurveModel* Curve = FindCurve(ID))
 		{
-			TArray<FKeyHandle> KeyHandles;
-			Curve->GetKeys(TNumericLimits<double>::Lowest(), TNumericLimits<double>::Max(), TNumericLimits<double>::Lowest(), TNumericLimits<double>::Max(), KeyHandles);
+			const TArray<FKeyHandle> KeyHandles = Curve->GetAllKeys();
 			Selection.Add(ID, ECurvePointType::Key, KeyHandles);
 		}
 	}
@@ -1126,8 +1125,7 @@ void FCurveEditor::InvertSelection()
 		FCurveModelID CurveModelID = Pair.Key;
 		if (FCurveModel* Curve = FindCurve(CurveModelID))
 		{
-			TArray<FKeyHandle> KeyHandles;
-			Curve->GetKeys(TNumericLimits<double>::Lowest(), TNumericLimits<double>::Max(), TNumericLimits<double>::Lowest(), TNumericLimits<double>::Max(), KeyHandles);
+			TArray<FKeyHandle> KeyHandles = Curve->GetAllKeys();
 			
 			TArrayView<const FKeyHandle> SelectedKeyHandles = Pair.Value.AsArray();
 				
@@ -1300,8 +1298,7 @@ void FCurveEditor::FlipCurve(ECurveFlipDirection Direction)
 		if (FCurveModel* Curve = Pair.Value.Get())
 		{
 			// Init key handles
-			TArray<FKeyHandle> KeyHandles;
-			Curve->GetKeys(TNumericLimits<double>::Lowest(), TNumericLimits<double>::Max(), TNumericLimits<double>::Lowest(), TNumericLimits<double>::Max(), KeyHandles);
+			const TArray<FKeyHandle> KeyHandles = Curve->GetAllKeys();
 
 			// Init key positions
 			TArray<FKeyPosition> AllKeyPositions;
@@ -1570,7 +1567,7 @@ void FCurveEditor::MatchLastTangentToFirst(bool bMatchLastToFirst)
 		if (CurveModel)
 		{
 			// Get all of the key handles from this curve.
-			CurveModel->GetKeys(TNumericLimits<double>::Lowest(), TNumericLimits<double>::Max(), TNumericLimits<double>::Lowest(), TNumericLimits<double>::Max(), AllKeyHandles);
+			CurveModel->GetAllKeys(AllKeyHandles);
 			//need 2
 			if (AllKeyHandles.Num() < 2)
 			{
@@ -2343,8 +2340,7 @@ void FCurveEditor::ApplyBufferedCurveToTarget(const IBufferedCurveModel* Buffere
 
 
 	// Copy the data from the Buffered curve into the target curve. This just does wholesale replacement.
-	TArray<FKeyHandle> TargetKeyHandles;
-	TargetCurve->GetKeys(TNumericLimits<double>::Lowest(), TNumericLimits<double>::Max(), TNumericLimits<double>::Lowest(), TNumericLimits<double>::Max(), TargetKeyHandles);
+	const TArray<FKeyHandle> TargetKeyHandles = TargetCurve->GetAllKeys();
 
 	double CurrentTime = 0.0;
 	if (TSharedPtr<ITimeSliderController> TimeSliderController = WeakTimeSliderController.Pin())
@@ -2627,8 +2623,7 @@ void FCurveEditor::PostUndo(bool bSuccess)
 			continue;
 		}
 		// Get all of the key handles from this curve.
-		TArray<FKeyHandle> KeyHandles;
-		CurveModel->GetKeys(TNumericLimits<double>::Lowest(), TNumericLimits<double>::Max(), TNumericLimits<double>::Lowest(), TNumericLimits<double>::Max(), KeyHandles);
+		const TArray<FKeyHandle> KeyHandles = CurveModel->GetAllKeys();
 
 		// The set handles will be mutated as we remove things so we need a copy that we can iterate through.
 		TArrayView<const FKeyHandle> SelectedHandles = Set.Value.AsArray();
diff --git a/Engine/Source/Editor/CurveEditor/Private/Tree/SCurveEditorTreeSelect.cpp b/Engine/Source/Editor/CurveEditor/Private/Tree/SCurveEditorTreeSelect.cpp
--- a/Engine/Source/Editor/CurveEditor/Private/Tree/SCurveEditorTreeSelect.cpp
+++ b/Engine/Source/Editor/CurveEditor/Private/Tree/SCurveEditorTreeSelect.cpp
@@ -72,9 +72,7 @@ FReply SCurveEditorTreeSelect::SelectAll()
 		FCurveModel* CurveModel = CurveEditor.Get()->FindCurve(CurveID);
 		if (CurveModel)
 		{
-			TArray<FKeyHandle> KeyHandles;
-			KeyHandles.Reserve(CurveModel->GetNumKeys());
-			CurveModel->GetKeys(TNumericLimits<double>::Lowest(), TNumericLimits<double>::Max(), TNumericLimits<double>::Lowest(), TNumericLimits<double>::Max(), KeyHandles);
+			const TArray<FKeyHandle> KeyHandles = CurveModel->GetAllKeys();
 		
 			if (bIsControlDown)
 			{
diff --git a/Engine/Source/Editor/CurveEditor/Private/Views/SInteractiveCurveEditorView.cpp b/Engine/Source/Editor/CurveEditor/Private/Views/SInteractiveCurveEditorView.cpp
--- a/Engine/Source/Editor/CurveEditor/Private/Views/SInteractiveCurveEditorView.cpp
+++ b/Engine/Source/Editor/CurveEditor/Private/Views/SInteractiveCurveEditorView.cpp
@@ -1186,10 +1186,7 @@ bool SInteractiveCurveEditorView::GetCurveWithinWidgetRange(const FSlateRect& Wi
 	{
 		if (const FCurveModel* Curve = CurveEditor->FindCurve(CurveID))
 		{
-			TArray<FKeyHandle> KeyHandles;
-			Curve->GetKeys(TNumericLimits<double>::Lowest(), TNumericLimits<double>::Max(), TNumericLimits<double>::Lowest(), TNumericLimits<double>::Max(), KeyHandles);
-
-			for (const FKeyHandle& KeyHandle : KeyHandles)
+			for (const FKeyHandle& KeyHandle : Curve->GetAllKeys())
 			{
 				OutPoints->Add(FCurvePointHandle(CurveID, ECurvePointType::Key, KeyHandle));
 				bPointsAdded = true;
@@ -1814,9 +1811,7 @@ FReply SInteractiveCurveEditorView::OnMouseButtonUp(const FGeometry& MyGeometry,
 			const UE::CurveEditor::FScopedSelectionChange Transaction(WeakCurveEditor, LOCTEXT("ClickCurve", "Click Curve"));
 			FCurveModel* CurveModel = CurveEditor->FindCurve(HitCurve.GetValue());
 
-			TArray<FKeyHandle> KeyHandles;
-			KeyHandles.Reserve(CurveModel->GetNumKeys());
-			CurveModel->GetKeys(TNumericLimits<double>::Lowest(), TNumericLimits<double>::Max(), TNumericLimits<double>::Lowest(), TNumericLimits<double>::Max(), KeyHandles);
+			const TArray<FKeyHandle> KeyHandles = CurveModel->GetAllKeys();
 
 			// Add or remove all keys from the curve.
 			if (bIsShiftDown)
diff --git a/Engine/Source/Editor/CurveEditor/Public/Filters/CurveEditorFilterBase.h b/Engine/Source/Editor/CurveEditor/Public/Filters/CurveEditorFilterBase.h
--- a/Engine/Source/Editor/CurveEditor/Public/Filters/CurveEditorFilterBase.h
+++ b/Engine/Source/Editor/CurveEditor/Public/Filters/CurveEditorFilterBase.h
@@ -123,11 +123,7 @@ public:
 			check(CurveModel);
 
 			FKeyHandleSet& HandleSet = KeysToOperateOn.Add(CurveModelID);
-			
-			TArray<FKeyHandle> Handles;
-			CurveModel->GetKeys(TNumericLimits<double>::Lowest(), TNumericLimits<double>::Max(), TNumericLimits<double>::Lowest(), TNumericLimits<double>::Max(), Handles);
-
-			for (const FKeyHandle& Handle : Handles)
+			for (const FKeyHandle& Handle : CurveModel->GetAllKeys())
 			{
 				HandleSet.Add(Handle, ECurvePointType::Key);
 			}
diff --git a/Engine/Source/Editor/MovieSceneTools/Private/Cache/MovieSceneCachedCurve.h b/Engine/Source/Editor/MovieSceneTools/Private/Cache/MovieSceneCachedCurve.h
--- a/Engine/Source/Editor/MovieSceneTools/Private/Cache/MovieSceneCachedCurve.h
+++ b/Engine/Source/Editor/MovieSceneTools/Private/Cache/MovieSceneCachedCurve.h
@@ -409,10 +409,7 @@ namespace UE::MovieSceneTools
 		}
 
 		// Key Handles
-		TArray<FKeyHandle> NewKeyHandles;
-		NewKeyHandles.Reserve(Times.Num());
-
-		UpdateData.CurveModel.GetKeys(TNumericLimits<double>::Lowest(), TNumericLimits<double>::Max(), TNumericLimits<double>::Lowest(), TNumericLimits<double>::Max(), NewKeyHandles);
+		TArray<FKeyHandle> NewKeyHandles = UpdateData.CurveModel.GetAllKeys();
 		if (KeyHandles != NewKeyHandles)
 		{
 			KeyHandles = NewKeyHandles;
