Curve Editor: Panning

Added Right or Alt-Middle as the default option for Panning, which allows Right and Alt-Middle to both pan

diff --git a/Engine/Source/Editor/CurveEditor/Private/CurveEditorSettings.cpp b/Engine/Source/Editor/CurveEditor/Private/CurveEditorSettings.cpp
--- a/Engine/Source/Editor/CurveEditor/Private/CurveEditorSettings.cpp
+++ b/Engine/Source/Editor/CurveEditor/Private/CurveEditorSettings.cpp
@@ -8,7 +8,7 @@
 
 UCurveEditorSettings::UCurveEditorSettings()
 {
-	PanningMouseButton = ECurveEditorPanningMouseButton::Right;
+	PanningMouseButton = ECurveEditorPanningMouseButton::RightOrAltMiddle;
 	bAutoFrameCurveEditor = true;
 	bShowBars = true;
 	FrameInputPadding = 50;
@@ -38,6 +38,15 @@ bool UCurveEditorSettings::AllowMousePan(const FPointerEvent& MouseEvent) const
 	{
 		return MouseEvent.IsMouseButtonDown(EKeys::RightMouseButton) && !MouseEvent.IsAltDown();
 	}
+	else if (PanningMouseButton == ECurveEditorPanningMouseButton::RightOrAltMiddle)
+	{
+		if ((MouseEvent.IsMouseButtonDown(EKeys::MiddleMouseButton) && MouseEvent.IsAltDown()) ||
+			(MouseEvent.IsMouseButtonDown(EKeys::RightMouseButton) && !MouseEvent.IsAltDown()))
+		{
+			return true;
+		}
+		return false;
+	}
 
 	// Default to right mouse button
 	return MouseEvent.GetEffectingButton() == EKeys::RightMouseButton;
@@ -55,6 +64,14 @@ bool UCurveEditorSettings::AllowMouseEdit(const FPointerEvent& MouseEvent) const
 	{
 		return MouseEvent.IsMouseButtonDown(EKeys::MiddleMouseButton);
 	}
+	else if (PanningMouseButton == ECurveEditorPanningMouseButton::RightOrAltMiddle)
+	{
+		if (MouseEvent.IsMouseButtonDown(EKeys::MiddleMouseButton) && !MouseEvent.IsAltDown())
+		{
+			return true;
+		}
+		return false;
+	}
 
 	// Default to middle mouse button
 	return MouseEvent.GetEffectingButton() == EKeys::MiddleMouseButton;
diff --git a/Engine/Source/Editor/CurveEditor/Public/CurveEditorSettings.h b/Engine/Source/Editor/CurveEditor/Public/CurveEditorSettings.h
--- a/Engine/Source/Editor/CurveEditor/Public/CurveEditorSettings.h
+++ b/Engine/Source/Editor/CurveEditor/Public/CurveEditorSettings.h
@@ -25,7 +25,8 @@ UENUM()
 enum class ECurveEditorPanningMouseButton : uint8
 {
 	Right	  UMETA(DisplayName = "Right"),
-	AltMiddle UMETA(DisplayName = "Alt-Middle")
+	AltMiddle UMETA(DisplayName = "Alt-Middle"),
+	RightOrAltMiddle      UMETA(DisplayName = "Right or Alt-Middle")
 };
 
 /** Defines visibility states for the tangents in the curve editor. */
