Curve editor: Rename FScopedKeyChange to FScopedCurveChange to better communicate that any changes are detected.
Specically, FCurveAttributes changes are also detected. In the future, we may want to add detection of curve color, name, etc. so we don't want the name to be misleading.

diff --git a/Engine/Plugins/Animation/TweeningUtils/Source/TweeningUtilsEditor/Public/Math/Models/TransactionalTweenModelProxy.h b/Engine/Plugins/Animation/TweeningUtils/Source/TweeningUtilsEditor/Public/Math/Models/TransactionalTweenModelProxy.h
--- a/Engine/Plugins/Animation/TweeningUtils/Source/TweeningUtilsEditor/Public/Math/Models/TransactionalTweenModelProxy.h
+++ b/Engine/Plugins/Animation/TweeningUtils/Source/TweeningUtilsEditor/Public/Math/Models/TransactionalTweenModelProxy.h
@@ -8,7 +8,7 @@
 #include "Internationalization/Internationalization.h"
 #include "Internationalization/Text.h"
 #include "Misc/Attribute.h"
-#include "Modification/Utils/ScopedKeyChange.h"
+#include "Modification/Utils/ScopedCurveChange.h"
 #include "Templates/UniquePtr.h"
 
 #include <type_traits>
@@ -55,7 +55,7 @@ private:
 	const TAttribute<TWeakPtr<FCurveEditor>> CurveEditorAttr;
 	
 	/** Active during a blend operation. */
-	TOptional<CurveEditor::FScopedKeyChange> InProgressTransaction;
+	TOptional<CurveEditor::FScopedCurveChange> InProgressTransaction;
 };
 }
 
diff --git a/Engine/Plugins/Editor/CurveEditorTools/Source/CurveEditorTools/Private/Tools/CurveEditorMultiScaleTool.cpp b/Engine/Plugins/Editor/CurveEditorTools/Source/CurveEditorTools/Private/Tools/CurveEditorMultiScaleTool.cpp
--- a/Engine/Plugins/Editor/CurveEditorTools/Source/CurveEditorTools/Private/Tools/CurveEditorMultiScaleTool.cpp
+++ b/Engine/Plugins/Editor/CurveEditorTools/Source/CurveEditorTools/Private/Tools/CurveEditorMultiScaleTool.cpp
@@ -11,7 +11,7 @@
 #include "Algo/MaxElement.h"
 #include "Misc/Mirror/TangentSelectionFlattener.h"
 #include "Modification/Keys/CurvesSnapshotBuilder.h"
-#include "Modification/Utils/ScopedKeyChange.h"
+#include "Modification/Utils/ScopedCurveChange.h"
 
 #include UE_INLINE_GENERATED_CPP_BY_NAME(CurveEditorMultiScaleTool)
 
@@ -387,7 +387,7 @@ void FCurveEditorMultiScaleTool::OnDragStart()
 	}
 
 	using namespace UE::CurveEditor;
-	ActiveTransaction = MakeUnique<FScopedKeyChange>(
+	ActiveTransaction = MakeUnique<FScopedCurveChange>(
 		FCurvesSnapshotBuilder(CurveEditor, ECurveChangeFlags::KeyAttributes | ECurveChangeFlags::MoveKeys).TrackSelectedCurves(),
 		LOCTEXT("CurveEditorMultiScaleToolTransaction", "Multi Scale Key(s)")
 		);
diff --git a/Engine/Plugins/Editor/CurveEditorTools/Source/CurveEditorTools/Private/Tools/CurveEditorMultiScaleTool.h b/Engine/Plugins/Editor/CurveEditorTools/Source/CurveEditorTools/Private/Tools/CurveEditorMultiScaleTool.h
--- a/Engine/Plugins/Editor/CurveEditorTools/Source/CurveEditorTools/Private/Tools/CurveEditorMultiScaleTool.h
+++ b/Engine/Plugins/Editor/CurveEditorTools/Source/CurveEditorTools/Private/Tools/CurveEditorMultiScaleTool.h
@@ -7,7 +7,7 @@
 #include "Framework/DelayedDrag.h"
 #include "ICurveEditorToolExtension.h"
 #include "Misc/Mirror/TangentSelectionFlattener.h"
-#include "Modification/Utils/ScopedKeyChange.h"
+#include "Modification/Utils/ScopedCurveChange.h"
 #include "UObject/StructOnScope.h"
 #include "CurveEditorMultiScaleTool.generated.h"
 
@@ -133,7 +133,7 @@ private:
 	TWeakPtr<FCurveEditor> WeakCurveEditor;
 
 	/** The currently open transaction (if any) */
-	TUniquePtr<UE::CurveEditor::FScopedKeyChange> ActiveTransaction;
+	TUniquePtr<UE::CurveEditor::FScopedCurveChange> ActiveTransaction;
 	
 	/** Set when attempting to move a drag handle. */
 	TOptional<FDelayedDrag> DelayedDrag;
diff --git a/Engine/Plugins/Editor/CurveEditorTools/Source/CurveEditorTools/Private/Tools/CurveEditorRetimeTool.cpp b/Engine/Plugins/Editor/CurveEditorTools/Source/CurveEditorTools/Private/Tools/CurveEditorRetimeTool.cpp
--- a/Engine/Plugins/Editor/CurveEditorTools/Source/CurveEditorTools/Private/Tools/CurveEditorRetimeTool.cpp
+++ b/Engine/Plugins/Editor/CurveEditorTools/Source/CurveEditorTools/Private/Tools/CurveEditorRetimeTool.cpp
@@ -13,7 +13,7 @@
 #include "Rendering/SlateRenderer.h"
 #include "Fonts/FontMeasure.h"
 #include "SCurveEditorView.h"
-#include "Modification/Utils/ScopedKeyChange.h"
+#include "Modification/Utils/ScopedCurveChange.h"
 #include "Styling/CoreStyle.h"
 
 #include UE_INLINE_GENERATED_CPP_BY_NAME(CurveEditorRetimeTool)
diff --git a/Engine/Plugins/Editor/CurveEditorTools/Source/CurveEditorTools/Private/Tools/CurveEditorRetimeTool.h b/Engine/Plugins/Editor/CurveEditorTools/Source/CurveEditorTools/Private/Tools/CurveEditorRetimeTool.h
--- a/Engine/Plugins/Editor/CurveEditorTools/Source/CurveEditorTools/Private/Tools/CurveEditorRetimeTool.h
+++ b/Engine/Plugins/Editor/CurveEditorTools/Source/CurveEditorTools/Private/Tools/CurveEditorRetimeTool.h
@@ -6,7 +6,7 @@
 #include "Framework/DelayedDrag.h"
 #include "CurveDataAbstraction.h"
 #include "CurveEditorSelection.h"
-#include "Modification/Utils/ScopedKeyChange.h"
+#include "Modification/Utils/ScopedCurveChange.h"
 #include "ScopedTransaction.h"
 #include "CurveEditorRetimeTool.generated.h"
 
@@ -109,7 +109,7 @@ private:
 	struct FTransaction
 	{
 		FScopedTransaction Transaction;
-		UE::CurveEditor::FScopedKeyChange KeyChange;
+		UE::CurveEditor::FScopedCurveChange KeyChange;
 
 		explicit FTransaction(TWeakPtr<FCurveEditor> InCurveEditor, const FText& Description)
 			: Transaction(Description)
diff --git a/Engine/Plugins/Editor/CurveEditorTools/Source/CurveEditorTools/Private/Tools/CurveEditorTransformTool.cpp b/Engine/Plugins/Editor/CurveEditorTools/Source/CurveEditorTools/Private/Tools/CurveEditorTransformTool.cpp
--- a/Engine/Plugins/Editor/CurveEditorTools/Source/CurveEditorTools/Private/Tools/CurveEditorTransformTool.cpp
+++ b/Engine/Plugins/Editor/CurveEditorTools/Source/CurveEditorTools/Private/Tools/CurveEditorTransformTool.cpp
@@ -787,7 +787,7 @@ void FCurveEditorTransformTool::OnDragStart()
 	}
 
 	using namespace UE::CurveEditor;
-	ActiveTransaction = MakeUnique<FScopedKeyChange>(
+	ActiveTransaction = MakeUnique<FScopedCurveChange>(
 		FCurvesSnapshotBuilder(CurveEditor, ECurveChangeFlags::KeyAttributes | ECurveChangeFlags::MoveKeys).TrackSelectedCurves(),
 		LOCTEXT("CurveEditorTransformToolTransaction", "Transform Key(s)")
 		);
diff --git a/Engine/Plugins/Editor/CurveEditorTools/Source/CurveEditorTools/Private/Tools/CurveEditorTransformTool.h b/Engine/Plugins/Editor/CurveEditorTools/Source/CurveEditorTools/Private/Tools/CurveEditorTransformTool.h
--- a/Engine/Plugins/Editor/CurveEditorTools/Source/CurveEditorTools/Private/Tools/CurveEditorTransformTool.h
+++ b/Engine/Plugins/Editor/CurveEditorTools/Source/CurveEditorTools/Private/Tools/CurveEditorTransformTool.h
@@ -9,7 +9,7 @@
 #include "ScopedTransaction.h"
 #include "CurveEditorSnapMetrics.h"
 #include "Misc/Mirror/TangentSelectionFlattener.h"
-#include "Modification/Utils/ScopedKeyChange.h"
+#include "Modification/Utils/ScopedCurveChange.h"
 #include "CurveEditorTransformTool.generated.h"
 
 class FCurveEditor;
@@ -157,7 +157,7 @@ private:
 	TWeakPtr<FCurveEditor> WeakCurveEditor;
 
 	/** The currently open transaction (if any) */
-	TUniquePtr<UE::CurveEditor::FScopedKeyChange> ActiveTransaction;
+	TUniquePtr<UE::CurveEditor::FScopedCurveChange> ActiveTransaction;
 
 	/** Cached information about our transform tool such as interaction state, etc. */
 	FCurveEditorTransformWidget TransformWidget;
diff --git a/Engine/Plugins/Editor/CurveEditorTools/Source/CurveEditorTools/Private/Tools/Lattice/CurveEditorLatticeTool.cpp b/Engine/Plugins/Editor/CurveEditorTools/Source/CurveEditorTools/Private/Tools/Lattice/CurveEditorLatticeTool.cpp
--- a/Engine/Plugins/Editor/CurveEditorTools/Source/CurveEditorTools/Private/Tools/Lattice/CurveEditorLatticeTool.cpp
+++ b/Engine/Plugins/Editor/CurveEditorTools/Source/CurveEditorTools/Private/Tools/Lattice/CurveEditorLatticeTool.cpp
@@ -121,7 +121,7 @@ FReply FCurveEditorLatticeTool::OnMouseButtonDoubleClick(
 			if (HoverState.HoveredCell)
 			{
 				const FScopedTransaction Transaction(LOCTEXT("FlipCurves", "Flip curves"));
-				const FScopedKeyChange KeyChange(
+				const FScopedCurveChange KeyChange(
 					FCurvesSnapshotBuilder(WeakCurveEditor, ECurveChangeFlags::KeyAttributes | ECurveChangeFlags::MoveKeys).TrackSelectedCurves()
 					);
 				UndoObject->Modify();
@@ -133,7 +133,7 @@ FReply FCurveEditorLatticeTool::OnMouseButtonDoubleClick(
 			if (HoverState.HoveredEdge)
 			{
 				const FScopedTransaction Transaction(LOCTEXT("FlattenCurves", "Flatten curves"));
-				const FScopedKeyChange KeyChange(
+				const FScopedCurveChange KeyChange(
 					FCurvesSnapshotBuilder(WeakCurveEditor, ECurveChangeFlags::KeyAttributes | ECurveChangeFlags::MoveKeys).TrackSelectedCurves()
 					);
 				UndoObject->Modify();
diff --git a/Engine/Plugins/Editor/CurveEditorTools/Source/CurveEditorTools/Private/Tools/Lattice/CurveEditorLatticeTool.h b/Engine/Plugins/Editor/CurveEditorTools/Source/CurveEditorTools/Private/Tools/Lattice/CurveEditorLatticeTool.h
--- a/Engine/Plugins/Editor/CurveEditorTools/Source/CurveEditorTools/Private/Tools/Lattice/CurveEditorLatticeTool.h
+++ b/Engine/Plugins/Editor/CurveEditorTools/Source/CurveEditorTools/Private/Tools/Lattice/CurveEditorLatticeTool.h
@@ -16,7 +16,7 @@
 #include "Misc/LatticeDrawUtils.h"
 #include "PerCurveLatticeData.h"
 #include "ScopedTransaction.h"
-#include "Modification/Utils/ScopedKeyChange.h"
+#include "Modification/Utils/ScopedCurveChange.h"
 #include "Templates/SharedPointer.h"
 #include "UObject/UnrealType.h"
 
@@ -141,7 +141,7 @@ private:
 		/** This transaction is active for the duration of the drag. */
 		TUniquePtr<FScopedTransaction> Transaction;
 		/** Transactions the key changes */
-		TUniquePtr<CurveEditor::FScopedKeyChange> KeyChange;
+		TUniquePtr<CurveEditor::FScopedCurveChange> KeyChange;
 		
 		/** Adjust tangents while dragging on an edge.*/
 		TMap<FCurveModelID, FLatticeEdgeTangentsMirrorOp> EdgeTangentMirroringOps;
@@ -157,7 +157,7 @@ private:
 			, KeyChange([InCurveEditor]
 			{
 				using namespace UE::CurveEditor;
-				return MakeUnique<FScopedKeyChange>(
+				return MakeUnique<FScopedCurveChange>(
 					FCurvesSnapshotBuilder(InCurveEditor, ECurveChangeFlags::MoveKeys | ECurveChangeFlags::KeyAttributes).TrackSelectedCurves()
 					);
 			}())
diff --git a/Engine/Source/Editor/CurveEditor/Private/CurveEditor.cpp b/Engine/Source/Editor/CurveEditor/Private/CurveEditor.cpp
--- a/Engine/Source/Editor/CurveEditor/Private/CurveEditor.cpp
+++ b/Engine/Source/Editor/CurveEditor/Private/CurveEditor.cpp
@@ -60,7 +60,7 @@
 #include "UObject/UnrealNames.h"
 #include "UnrealExporter.h"
 #include "Misc/SmartSnap.h"
-#include "Modification/Utils/ScopedKeyChange.h"
+#include "Modification/Utils/ScopedCurveChange.h"
 #include "Modification/Utils/ScopedSelectionChange.h"
 #include "Widgets/Colors/SColorPicker.h"
 #include "Widgets/Notifications/SNotificationList.h"
@@ -1289,7 +1289,7 @@ void FCurveEditor::FlipCurveVertical(TArray<FKeyPosition>& AllKeyPositions, TArr
 void FCurveEditor::FlipCurve(ECurveFlipDirection Direction)
 {
 	using namespace UE::CurveEditor;
-	const FScopedKeyChange KeyChange(
+	const FScopedCurveChange KeyChange(
 		FCurvesSnapshotBuilder(SharedThis(this), CurveData, ECurveChangeFlags::MoveKeys | ECurveChangeFlags::KeyAttributes),
 		LOCTEXT("FlipCurve", "Flip Curve")
 		);
@@ -1549,7 +1549,7 @@ void FCurveEditor::MatchLastTangentToFirst(bool bMatchLastToFirst)
 {
 	using namespace UE::CurveEditor;
 	// FScopedKeyChange won't generate any transaction if no tangents are changed
-	const FScopedKeyChange KeyChange(
+	const FScopedCurveChange KeyChange(
 		FCurvesSnapshotBuilder(SharedThis(this), CurveData, ECurveChangeFlags::KeyAttributes),
 		LOCTEXT("MatchTangents", "Match Tangents")
 		);
@@ -1850,7 +1850,7 @@ void FCurveEditor::PasteKeys(UE::CurveEditor::FKeyPasteArgs InArgs)
 	const FScopedTransaction Transaction(LOCTEXT("PasteKeys", "Paste Keys")); // Selection and key change should be part of same undo / redo op.
 	const FScopedSelectionChange SelectionChange(SharedThis(this));
 	constexpr ECurveChangeFlags Flags = ECurveChangeFlags::AddKeys | ECurveChangeFlags::RemoveKeys | ECurveChangeFlags::KeyAttributes;
-	const FScopedKeyChange KeyChange(FCurvesSnapshotBuilder(SharedThis(this), TargetCurves, Flags));
+	const FScopedCurveChange KeyChange(FCurvesSnapshotBuilder(SharedThis(this), TargetCurves, Flags));
 	Selection.Clear();
 
 	// Two simple cases, 1 to 1 and 1 to many.
@@ -1984,7 +1984,7 @@ void FCurveEditor::DeleteSelection()
 	
 	const FScopedTransaction Transaction(LOCTEXT("DeleteKeys", "Delete Keys"));
 	const FScopedSelectionChange SelectionChange(SharedThis(this));
-	const FScopedKeyChange KeyChange(
+	const FScopedCurveChange KeyChange(
 		FCurvesSnapshotBuilder(SharedThis(this), ECurveChangeFlags::RemoveKeys).TrackSelectedCurves()
 		);
 	
@@ -2010,7 +2010,7 @@ void FCurveEditor::FlattenSelection()
 {
 	using namespace UE::CurveEditor;
 	// FScopedKeyChange won't generate any transaction if no tangents are changed
-	const FScopedKeyChange KeyChange(
+	const FScopedCurveChange KeyChange(
 		FCurvesSnapshotBuilder(SharedThis(this), CurveData, ECurveChangeFlags::KeyAttributes),
 		LOCTEXT("FlattenTangents", "Flatten Tangents")
 		);
@@ -2093,7 +2093,7 @@ void FCurveEditor::StraightenSelection()
 {
 	using namespace UE::CurveEditor;
 	// FScopedKeyChange won't generate any transaction if no tangents are changed
-	const FScopedKeyChange KeyChange(
+	const FScopedCurveChange KeyChange(
 		FCurvesSnapshotBuilder(SharedThis(this), CurveData, ECurveChangeFlags::KeyAttributes),
 		LOCTEXT("StraightenTangents", "Straighten Tangents")
 		);
@@ -2148,7 +2148,7 @@ void FCurveEditor::SmartSnapSelection()
 	using namespace UE::CurveEditor;
 	const FScopedTransaction Transaction(LOCTEXT("SmartSnapKeys", "Smart Snap")); // Selection and key change should be part of same undo / redo op.
 	const FScopedSelectionChange SelectionChange(SharedThis(this));
-	const FScopedKeyChange KeyChange(
+	const FScopedCurveChange KeyChange(
 		FCurvesSnapshotBuilder(SharedThis(this), ECurveChangeFlags::MoveKeys | ECurveChangeFlags::RemoveKeys).TrackSelectedCurves()
 		);
 	
@@ -2359,7 +2359,7 @@ bool FCurveEditor::ApplyBufferedCurves(const TSet<FCurveModelID>& InCurvesToAppl
 {
 	using namespace UE::CurveEditor;
 	// Fyi, FScopedKeyChange won't make any entry in the Undo/Redo buffer if nothing was changed.
-	const FScopedKeyChange KeyChange(
+	const FScopedCurveChange KeyChange(
 		FCurvesSnapshotBuilder(SharedThis(this), InCurvesToApplyTo),
 		bSwapBufferCurves ? LOCTEXT("SwapBufferedCurves", "Swap Buffered Curves") : LOCTEXT("ApplyBufferedCurves", "Apply Buffered Curves")
 		);
diff --git a/Engine/Source/Editor/CurveEditor/Private/DragOperations/CurveEditorDragOperation_MoveKeys.cpp b/Engine/Source/Editor/CurveEditor/Private/DragOperations/CurveEditorDragOperation_MoveKeys.cpp
--- a/Engine/Source/Editor/CurveEditor/Private/DragOperations/CurveEditorDragOperation_MoveKeys.cpp
+++ b/Engine/Source/Editor/CurveEditor/Private/DragOperations/CurveEditorDragOperation_MoveKeys.cpp
@@ -17,7 +17,7 @@
 #include "Misc/AssertionMacros.h"
 #include "SCurveEditorView.h"
 #include "ScopedTransaction.h"
-#include "Modification/Utils/ScopedKeyChange.h"
+#include "Modification/Utils/ScopedCurveChange.h"
 #include "Templates/Tuple.h"
 #include "Templates/UniquePtr.h"
 #include "Templates/UnrealTemplate.h"
diff --git a/Engine/Source/Editor/CurveEditor/Private/DragOperations/CurveEditorDragOperation_MoveKeys.h b/Engine/Source/Editor/CurveEditor/Private/DragOperations/CurveEditorDragOperation_MoveKeys.h
--- a/Engine/Source/Editor/CurveEditor/Private/DragOperations/CurveEditorDragOperation_MoveKeys.h
+++ b/Engine/Source/Editor/CurveEditor/Private/DragOperations/CurveEditorDragOperation_MoveKeys.h
@@ -10,7 +10,7 @@
 #include "ICurveEditorDragOperation.h"
 #include "Math/Vector2D.h"
 #include "Misc/Optional.h"
-#include "Modification/Utils/ScopedKeyChange.h"
+#include "Modification/Utils/ScopedCurveChange.h"
 
 class FCurveEditor;
 struct FPointerEvent;
@@ -37,7 +37,7 @@ private:
 	FVector2D LastMousePosition;
 
 	/** Transacts the moved keys. */
-	TOptional<UE::CurveEditor::FScopedKeyChange> ScopedKeyChange;
+	TOptional<UE::CurveEditor::FScopedCurveChange> ScopedKeyChange;
 
 	struct FKeyData
 	{
diff --git a/Engine/Source/Editor/CurveEditor/Private/DragOperations/CurveEditorDragOperation_Tangent.h b/Engine/Source/Editor/CurveEditor/Private/DragOperations/CurveEditorDragOperation_Tangent.h
--- a/Engine/Source/Editor/CurveEditor/Private/DragOperations/CurveEditorDragOperation_Tangent.h
+++ b/Engine/Source/Editor/CurveEditor/Private/DragOperations/CurveEditorDragOperation_Tangent.h
@@ -10,7 +10,7 @@
 #include "Math/Vector2D.h"
 #include "Misc/EnumClassFlags.h"
 #include "Misc/Optional.h"
-#include "Modification/Utils/ScopedKeyChange.h"
+#include "Modification/Utils/ScopedCurveChange.h"
 
 class FCurveEditor;
 struct FCurveEditorScreenSpace;
@@ -140,5 +140,5 @@ private:
 	TArray<FKeyData> KeysByCurve;
 	
 	/** Transacts the keys' tangents. */
-	TOptional<UE::CurveEditor::FScopedKeyChange> ScopedKeyChange;
+	TOptional<UE::CurveEditor::FScopedCurveChange> ScopedKeyChange;
 };
\ No newline at end of file
diff --git a/Engine/Source/Editor/CurveEditor/Private/Filters/FilterUtils.cpp b/Engine/Source/Editor/CurveEditor/Private/Filters/FilterUtils.cpp
--- a/Engine/Source/Editor/CurveEditor/Private/Filters/FilterUtils.cpp
+++ b/Engine/Source/Editor/CurveEditor/Private/Filters/FilterUtils.cpp
@@ -6,7 +6,7 @@
 #include "Modification/Utils/ScopedSelectionChange.h"
 
 #include "Containers/Map.h"
-#include "Modification/Utils/ScopedKeyChange.h"
+#include "Modification/Utils/ScopedCurveChange.h"
 #include "Templates/SharedPointer.h"
 
 namespace UE::CurveEditor::FilterUtils
@@ -26,7 +26,7 @@ namespace UE::CurveEditor::FilterUtils
 		
 		TMap<FCurveModelID, FKeyHandleSet> OutKeysToSelect;
 		{
-			const FScopedKeyChange KeyChange(FCurvesSnapshotBuilder(InCurveEditor, InSelectedKeys));
+			const FScopedCurveChange KeyChange(FCurvesSnapshotBuilder(InCurveEditor, InSelectedKeys));
 			InFilter.ApplyFilter(InCurveEditor, InSelectedKeys, OutKeysToSelect);
 		}
 		
diff --git a/Engine/Source/Editor/CurveEditor/Private/Modification/Utils/ScopedCurveChange.cpp b/Engine/Source/Editor/CurveEditor/Private/Modification/Utils/ScopedCurveChange.cpp
new file mode 100644
--- /dev/null
+++ b/Engine/Source/Editor/CurveEditor/Private/Modification/Utils/ScopedCurveChange.cpp
@@ -0,0 +1,88 @@
+﻿// Copyright Epic Games, Inc. All Rights Reserved.
+
+#include "Modification/Utils/ScopedCurveChange.h"
+
+#include "CurveEditor.h"
+#include "CurveModel.h"
+#include "HAL/IConsoleManager.h"
+#include "Internationalization/Text.h"
+#include "Misc/ITransaction.h"
+#include "Misc/ScopeExit.h"
+#include "Modification/Keys/Data/GenericCurveChangeData.h"
+#include "Modification/Keys/GenericCurveChangeCommand.h"
+#include "Modification/Keys/GenericCurveChangeUtils.h"
+#include "Modification/Keys/SnapshotPerfDebugUtils.h"
+#include "ProfilingDebugging/ScopedTimers.h"
+#include "SCurveEditor.h"
+
+namespace UE::CurveEditor
+{
+static TAutoConsoleVariable<bool> CVarLogUndoKeyChangeSize(
+	TEXT("CurveEditor.LogDiffPerfStats"),
+	false,
+	TEXT("Whenever a key change is made, whether to log debug info (time spent diffing, bytes used for the diff, the final bytes on undo stack size).")
+	);
+
+FScopedCurveChange::FScopedCurveChange(
+	TWeakPtr<FCurveEditor> InCurveEditor, FCurveChangeDiff InKeyChangeDiff, const FText& InDescription, EScopedKeyChangeFlags InFlags
+	)
+	: FScopedChangeBase(MoveTemp(InCurveEditor), InDescription.IsEmpty() ? NSLOCTEXT("CurveEditor", "ChangeKeys", "Change Keys") : InDescription)
+	, ChangeDiff(MoveTemp(InKeyChangeDiff))
+	, Flags(InFlags)
+{}
+
+FScopedCurveChange::FScopedCurveChange(
+	FCurvesSnapshotBuilder InDataToDiff, const FText& InDescription, EScopedKeyChangeFlags InFlags
+	)
+	: FScopedCurveChange(
+		InDataToDiff.WeakCurveEditor,
+		FCurveChangeDiff(InDataToDiff.WeakCurveEditor, MoveTemp(InDataToDiff.Snapshot), InDataToDiff.CapturedDataFlags),
+		InDescription, InFlags
+		)
+{}
+
+FScopedCurveChange::~FScopedCurveChange()
+{
+	const TSharedPtr<FCurveEditor> CurveEditorPin = WeakCurveEditor.Pin();
+	if (!CurveEditorPin)
+	{
+		return;
+	}
+
+	// Early out without expensive call to ComputeDiff.
+	const bool bRevertOnCancelled = EnumHasAnyFlags(Flags, EScopedKeyChangeFlags::RevertOnCancel);
+	const bool bHasBeenCancelled = IsCancelled();
+	if (bHasBeenCancelled && !bRevertOnCancelled)
+	{
+		return;
+	}
+
+	FAutoScopedDurationTimer TimeMeasurement;
+	FGenericCurveChangeData DeltaChange = ChangeDiff.ComputeDiff();
+	const auto LogPerfData = [this, &TimeMeasurement, &DeltaChange]
+	{
+		UE_CLOG(CVarLogUndoKeyChangeSize.GetValueOnAnyThread(), LogCurveEditor, Log, TEXT("Change'%s' --- Undo size: %s ---\t Diff: Time %f ms, %s"),
+			GUndo ? *GUndo->GetContext().Title.ToString() : *Description.ToString(),
+			*FText::AsMemory(DeltaChange.GetAllocatedSize()).ToString(),
+			FTimespan::FromSeconds(TimeMeasurement.GetTime()).GetTotalMilliseconds(),
+			*DumpSnapshotPerfData(ChangeDiff.GetSnapshot(), ECurvesSnapshotPerfFlags::All)
+			); 
+	};
+	
+	if (!DeltaChange)
+	{
+		return;
+	}
+	
+	if (!bHasBeenCancelled)
+	{
+		LogPerfData(); // Important to do this before MoveTemp
+		TryAppendCommand(MakeUnique<FGenericCurveChangeCommand>(MoveTemp(DeltaChange)));
+	}
+	else
+	{
+		LogPerfData();
+		GenericCurveChange::RevertChange(*CurveEditorPin, DeltaChange);
+	}
+}
+}
diff --git a/Engine/Source/Editor/CurveEditor/Private/Modification/Utils/ScopedKeyChange.cpp b/Engine/Source/Editor/CurveEditor/Private/Modification/Utils/ScopedKeyChange.cpp
deleted file mode 100644
diff --git a/Engine/Source/Editor/CurveEditor/Private/SCurveEditorPanel.cpp b/Engine/Source/Editor/CurveEditor/Private/SCurveEditorPanel.cpp
--- a/Engine/Source/Editor/CurveEditor/Private/SCurveEditorPanel.cpp
+++ b/Engine/Source/Editor/CurveEditor/Private/SCurveEditorPanel.cpp
@@ -67,7 +67,7 @@
 #include "Filters/CurveEditorEulerFilter.h"
 #include "Filters/PromotedFilterContainer.h"
 #include "Misc/ResizeParamUtils.h"
-#include "Modification/Utils/ScopedKeyChange.h"
+#include "Modification/Utils/ScopedCurveChange.h"
 #include "Modification/Utils/ScopedSelectionChange.h"
 #include "Styling/AppStyle.h"
 #include "Styling/CoreStyle.h"
@@ -1047,7 +1047,7 @@ EVisibility SCurveEditorPanel::GetSplitterVisibility() const
 void SCurveEditorPanel::SetKeyAttributes(FKeyAttributes KeyAttributes, FText Description)
 {
 	using namespace UE::CurveEditor;
-	const FScopedKeyChange KeyChange(
+	const FScopedCurveChange KeyChange(
 		FCurvesSnapshotBuilder(GetCurveEditor(), ECurveChangeFlags::KeyAttributes).TrackSelectedCurves(),
 		Description
 		);
@@ -1064,7 +1064,7 @@ void SCurveEditorPanel::SetKeyAttributes(FKeyAttributes KeyAttributes, FText Des
 void SCurveEditorPanel::SetCurveAttributes(FCurveAttributes CurveAttributes, FText Description)
 {
 	using namespace UE::CurveEditor;
-	const FScopedKeyChange KeyChange(
+	const FScopedCurveChange KeyChange(
 		FCurvesSnapshotBuilder(CurveEditor, ECurveChangeFlags::CurveAttributes).TrackSelectedCurves(),
 		Description
 		);
@@ -1082,7 +1082,7 @@ void SCurveEditorPanel::ToggleWeightedTangents()
 {
 	using namespace UE::CurveEditor;
 	const TMap<FCurveModelID, FKeyHandleSet>& Selection = CurveEditor->GetSelection().GetAll();
-	const FScopedKeyChange KeyChange(
+	const FScopedCurveChange KeyChange(
 		FCurvesSnapshotBuilder(CurveEditor, Selection, ECurveChangeFlags::KeyAttributes),
 		LOCTEXT("ToggleWeightedTangents_Transaction", "Toggle Weighted Tangents")
 		);
diff --git a/Engine/Source/Editor/CurveEditor/Private/Views/SInteractiveCurveEditorView.cpp b/Engine/Source/Editor/CurveEditor/Private/Views/SInteractiveCurveEditorView.cpp
--- a/Engine/Source/Editor/CurveEditor/Private/Views/SInteractiveCurveEditorView.cpp
+++ b/Engine/Source/Editor/CurveEditor/Private/Views/SInteractiveCurveEditorView.cpp
@@ -1626,7 +1626,7 @@ FReply SInteractiveCurveEditorView::OnMouseButtonDown(const FGeometry& MyGeometr
 			{
 				FScopedTransaction Transaction(LOCTEXT("InsertKey", "Insert Key")); // Selection and key change should be part of same undo / redo op.
 				const FScopedSelectionChange SelectionChange(CurveEditor);
-				const FScopedKeyChange KeyChange(
+				const FScopedCurveChange KeyChange(
 					FCurvesSnapshotBuilder(CurveEditor, *HoveredCurve, ECurveChangeFlags::KeyAttributes | ECurveChangeFlags::AddKeys)
 					);
 
@@ -2121,7 +2121,7 @@ void SInteractiveCurveEditorView::AddKeyAtTime(const TSet<FCurveModelID>& ToCurv
 
 	FScopedTransaction Transaction(LOCTEXT("AddKeyAtTime", "Add Key")); // Selection and key change should be part of same undo / redo op.
 	FScopedSelectionChange SelectionChange(WeakCurveEditor);
-	FScopedKeyChange KeyChange(
+	FScopedCurveChange KeyChange(
 		FCurvesSnapshotBuilder(WeakCurveEditor, ToCurves, ECurveChangeFlags::AddKeys)
 		);
 	
diff --git a/Engine/Source/Editor/CurveEditor/Public/Modification/Utils/ScopedCurveChange.h b/Engine/Source/Editor/CurveEditor/Public/Modification/Utils/ScopedCurveChange.h
new file mode 100644
--- /dev/null
+++ b/Engine/Source/Editor/CurveEditor/Public/Modification/Utils/ScopedCurveChange.h
@@ -0,0 +1,99 @@
+﻿// Copyright Epic Games, Inc. All Rights Reserved.
+
+#pragma once
+
+#include "Misc/EnumClassFlags.h"
+#include "Modification/Keys/CurveChangeDiff.h"
+#include "ScopedChangeBase.h"
+#include "Templates/SharedPointer.h"
+
+struct FCurveModelID;
+class FCurveEditor;
+
+namespace UE::CurveEditor
+{
+enum class EScopedKeyChangeFlags : uint8
+{
+	None,
+
+	/** Revert the change if FScopedKeyChange or surrounding FScopedTransaction becomes cancelled. */
+	RevertOnCancel = 1 << 0,
+	
+	Default = RevertOnCancel
+};
+ENUM_CLASS_FLAGS(EScopedKeyChangeFlags);
+
+enum class ETrackedCurves : uint8 { All, Selected };
+	
+/**
+ * Detects changes made to the curves specified on construction and adds it as an undo-able action.
+ * When the scope ends, checks whether keys were added, removed, moved, or its attributes (FKeyAttributes or FCurveAttributes) were edited.
+ *
+ * ========== Performance ==========
+ * This class exists for convenience, faster iteration times, and more concise code but introduces overhead by temporarily buffering the curves.
+ * If the overhead is too great, consider building the change directly, and appending it to the transaction manager manually.
+ * However, in practice, the overhead should not be a concern, and for simplicity, usage of FScopedKeyChange is encouraged and preferred.
+ * 
+ * Stress test 15th of July, 2025, Development Build,
+ * Sequencer: ControlRigExample -> CopyPaste_Level-> PossessableBindings_CopyPaste -> Select Possessable_FKCR_CopyPaste
+ *  - Set Interp Linear:	109.782 keys,	Undo Size: 5.582 MiB,		Diff buffer size: 2.955 MiB,	Diff time: 4.2 ms
+ *  - Move Keys:            109.782 keys,	Undo Size: 1.36 MiB,		Diff buffer size: 2.117 MiB,	Diff time: 5.3 ms
+ *	- Drag Tangents:        109.782 keys,	Undo Size: 5.582 MiB,		Diff buffer size: 2.955 MiB,	Diff time: 4.4 ms
+ *	- Blend values:         109.782 keys,	Undo Size: 6.583 MiB,		Diff buffer size: 4.654 MiB,	Diff time: 10.7 ms
+ *	- Delete Keys:          109.782 keys,	Undo Size: 5.395 MiB,		Diff buffer size: 4.63 MiB,		Diff time: 1.8 ms
+ * These are the stats logged by CurveEditor.LogDiffPerfStats 1.
+ *
+ * ========== Combining FScopedKeyChange and FScopedSelectionChange ==========
+ * The following is operations supported:
+ * - Combination with FScopedSelectionChange (see below),
+ * - Nesting within FScopedTransaction; if the outer FScopedTransaction is cancelled, FScopedKeyChange can detect this (@see EScopedKeyChangeFlags).
+ * - Nested FScopedKeyChanges
+ * 
+ * When combining FScopedKeyChange and FScopedSelectionChange, order matters.
+ * Rule of thumb: Declare in the order the undo operation should revert; equivalent: declare in reverse order the redo operation should apply.
+ * 
+* Keep the following principles in mind:
+ * - C++ will deconstruct in reverse order
+ * - Unreal's transaction system execution order for custom changes is as follows:
+ *	- Revert (undo) in insertion order. Below, the call order would be FSelectionChangeCommand::Revert, then FGenericCurveChangeCommand::Revert.
+ *	- Apply (redo) in reverse insertion order. Below, the call order would be FGenericCurveChangeCommand::Apply, then FSelectionChangeCommand::Apply.
+ * 
+ * Example code for inserting a key and selecting it:
+ *		FScopedSelectionChange SelectionChange; // When undoing, first remove from selection...
+ *		FScopedKeyChange KeyChange; // ... then remove the key from the curve
+ *		// Add key to some curve
+ *		// Select the added key
+ *		//~KeyChange > adds key change as undo action
+ *		//~SelectionChange > adds selection change as undo action
+ * If you changed the declaration order, on redo, we'd first add to selection and then add the key to the curve.
+ * Redo would no longer work: as FSelectionChangeCommand would add the FKeyHandle to the selection first, the selection code would defensively check
+ * that a key exists first and reject the function call as the key has not yet been added to the curve model.
+ */
+class FScopedCurveChange : public FScopedChangeBase
+{
+public:
+	
+	/**
+	 * @param InCurveEditor The curve
+	 * @param InKeyChangeDiff Contains the data that is supposed to be diffed against
+	 * @param InDescription Optional description that will be used for the transaction title. This title does not matter if you have a surrounding FScopedTransaction.
+	 * @param InFlags Extra flags to affect the behaviour of the change
+	 */
+	[[nodiscard]] CURVEEDITOR_API explicit FScopedCurveChange(
+		TWeakPtr<FCurveEditor> InCurveEditor, FCurveChangeDiff InKeyChangeDiff, const FText& InDescription, EScopedKeyChangeFlags InFlags
+		);
+	
+	[[nodiscard]] CURVEEDITOR_API explicit FScopedCurveChange(
+		FCurvesSnapshotBuilder InDataToDiff, const FText& InDescription = FText::GetEmpty(), EScopedKeyChangeFlags InFlags = EScopedKeyChangeFlags::Default
+		);
+
+	CURVEEDITOR_API ~FScopedCurveChange();
+
+private:
+	
+	/** Checks for changes. */
+	const FCurveChangeDiff ChangeDiff;
+
+	const EScopedKeyChangeFlags Flags;
+};
+}
diff --git a/Engine/Source/Editor/CurveEditor/Public/Modification/Utils/ScopedKeyChange.h b/Engine/Source/Editor/CurveEditor/Public/Modification/Utils/ScopedKeyChange.h
deleted file mode 100644
