Handle alt-tab capturelost for rectangle selection in curve editor.

#tests Tested with moving keys, tested with rectangle selection
#rnx

diff --git a/Engine/Source/Editor/CurveEditor/Private/SCurveEditorViewContainer.cpp b/Engine/Source/Editor/CurveEditor/Private/SCurveEditorViewContainer.cpp
--- a/Engine/Source/Editor/CurveEditor/Private/SCurveEditorViewContainer.cpp
+++ b/Engine/Source/Editor/CurveEditor/Private/SCurveEditorViewContainer.cpp
@@ -45,6 +45,13 @@ class FWidgetStyle;
 
 #define LOCTEXT_NAMESPACE "SCurveEditorViewContainer"
 
+#if WITH_EDITOR
+#include "Editor/EditorEngine.h"
+#include "EngineUtils.h"
+
+extern UNREALED_API UEditorEngine* GEditor;
+#endif // WITH_EDITOR
+
 void SCurveEditorViewContainer::Construct(const FArguments& InArgs, TSharedRef<FCurveEditor> InCurveEditor)
 {
 	SetCanTick(true);
@@ -393,6 +400,23 @@ FReply SCurveEditorViewContainer::OnMouseButtonUp(const FGeometry& MyGeometry, c
 	return Reply.ReleaseMouseCapture();
 }
 
+void SCurveEditorViewContainer::OnMouseCaptureLost(const FCaptureLostEvent& CaptureLostEvent)
+{
+#if WITH_EDITOR
+	if (!DragOperation.IsSet())
+	{
+		return;
+	}
+	GEditor->GetTimerManager()->SetTimerForNextTick([this]()
+		{
+			if (!HasMouseCapture())
+			{
+				DragOperation->DragImpl->CancelDrag();
+				DragOperation.Reset();
+			}
+		});
+#endif
+}
 FCursorReply SCurveEditorViewContainer::OnCursorQuery(const FGeometry& MyGeometry, const FPointerEvent& CursorEvent) const
 {
 	// Override the cursor while the Marquee select is happening. Ideally we'd be able to do a "+" and "-" while holding
diff --git a/Engine/Source/Editor/CurveEditor/Private/SCurveEditorViewContainer.h b/Engine/Source/Editor/CurveEditor/Private/SCurveEditorViewContainer.h
--- a/Engine/Source/Editor/CurveEditor/Private/SCurveEditorViewContainer.h
+++ b/Engine/Source/Editor/CurveEditor/Private/SCurveEditorViewContainer.h
@@ -76,6 +76,7 @@ private:
 	UE_API virtual FReply OnMouseButtonDoubleClick(const FGeometry& InMyGeometry, const FPointerEvent& InMouseEvent) override;
 	UE_API virtual void OnFocusLost(const FFocusEvent& InFocusEvent) override;
 	UE_API virtual FCursorReply OnCursorQuery(const FGeometry& MyGeometry, const FPointerEvent& CursorEvent) const override;
+	UE_API virtual void OnMouseCaptureLost(const FCaptureLostEvent& CaptureLostEvent) override;
 
 	virtual bool SupportsKeyboardFocus() const override { return true; }
 	UE_API virtual FReply OnKeyDown(const FGeometry& MyGeometry, const FKeyEvent& InKeyEvent) override;
