Curve Editor: Remove Modify call from SetColor.
- New: FCurveEditor::OnCurveColorsChangedDelegate
- SCurveEditorView subscribes to OnCurveColorsChangedDelegate and causes a redraw when the colors change
- This also fixes where the color would not be changed in other editors, like the curve asset editor. In the curve asset editor, you can now change the curve color (using View combo button). Before, these menu entries were non functional.

The Modify call in SetColor was being used to trigger a refresh of the view: the mechanism was that the internal Sequencer FGuid FChannelCurveModel::LastSignature would be changed... then when the UI calls HasChangedAndResetTest, it would detect a change and redraw. Instead, the locations that call SetColor, should refresh the view.

FCurvceEditor::SetCurveColorsForSelected and FCurveEditor::SetRandomCurveColorsForSelected have been changed to refresh the view by invoking OnCurveColorsChangedDelegate
It seems like those are the only places the user can uses to change the color. All other call sites of SetColor either operate on a FCurveModel that was just created (MakeUnique just before) or after loading the color from a config file (meaning that the color is not saved in any UObject meaning we can safely remove the Modify call - we do not need to mark the package dirty).

diff --git a/Engine/Plugins/Cameras/GameplayCameras/Source/GameplayCamerasEditor/Private/CurveEditor/CurvePropertyEditorTreeItem.cpp b/Engine/Plugins/Cameras/GameplayCameras/Source/GameplayCamerasEditor/Private/CurveEditor/CurvePropertyEditorTreeItem.cpp
--- a/Engine/Plugins/Cameras/GameplayCameras/Source/GameplayCamerasEditor/Private/CurveEditor/CurvePropertyEditorTreeItem.cpp
+++ b/Engine/Plugins/Cameras/GameplayCameras/Source/GameplayCamerasEditor/Private/CurveEditor/CurvePropertyEditorTreeItem.cpp
@@ -68,7 +68,7 @@ void FCurvePropertyEditorTreeItem::CreateCurveModels(TArray<TUniquePtr<FCurveMod
 	{
 		TUniquePtr<FRichCurveEditorModelRaw> NewCurve = MakeUnique<FRichCurveEditorModelRaw>(Info.Curve, Info.WeakOwner.Get());
 		NewCurve->SetShortDisplayName(Info.DisplayName);
-		NewCurve->SetColor(Info.Color, false);
+		NewCurve->SetColor(Info.Color);
 		OutCurveModels.Add(MoveTemp(NewCurve));
 	}
 }
diff --git a/Engine/Plugins/FX/Niagara/Source/NiagaraEditorWidgets/Private/DetailCustomizations/NiagaraDataInterfaceCurveDetails.cpp b/Engine/Plugins/FX/Niagara/Source/NiagaraEditorWidgets/Private/DetailCustomizations/NiagaraDataInterfaceCurveDetails.cpp
--- a/Engine/Plugins/FX/Niagara/Source/NiagaraEditorWidgets/Private/DetailCustomizations/NiagaraDataInterfaceCurveDetails.cpp
+++ b/Engine/Plugins/FX/Niagara/Source/NiagaraEditorWidgets/Private/DetailCustomizations/NiagaraDataInterfaceCurveDetails.cpp
@@ -147,8 +147,7 @@ struct FNiagaraCurveDetailsTreeItem : public ICurveEditorTreeItem, TSharedFromTh
 		{
 			NewCurve->SetShortDisplayName(FText::FromName(CurveData.Name));
 		}
-		bool bModify = false; // Don't want to modify (and mark dirty) the underlying data when creating a view model for it.
-		NewCurve->SetColor(CurveData.Color, bModify);
+		NewCurve->SetColor(CurveData.Color);
 		NewCurve->OnCurveModified().AddSP(this, &FNiagaraCurveDetailsTreeItem::CurveChanged);
 		OutCurveModels.Add(MoveTemp(NewCurve));
 	}
diff --git a/Engine/Source/Editor/CurveAssetEditor/Private/CurveAssetEditor.cpp b/Engine/Source/Editor/CurveAssetEditor/Private/CurveAssetEditor.cpp
--- a/Engine/Source/Editor/CurveAssetEditor/Private/CurveAssetEditor.cpp
+++ b/Engine/Source/Editor/CurveAssetEditor/Private/CurveAssetEditor.cpp
@@ -117,7 +117,7 @@ struct FCurveAssetEditorTreeItem : public ICurveEditorTreeItem
 
 		TUniquePtr<FRichCurveEditorModelRaw> NewCurve = MakeUnique<FRichCurveEditorModelRaw>(static_cast<FRichCurve*>(EditInfo.CurveToEdit), CurveOwner.Get());
 		NewCurve->SetShortDisplayName(CurveName);
-		NewCurve->SetColor(CurveColor, false);
+		NewCurve->SetColor(CurveColor);
 		OutCurveModels.Add(MoveTemp(NewCurve));
 	}
 	
diff --git a/Engine/Source/Editor/CurveAssetEditor/Private/SColorGradientCurveEditorView.cpp b/Engine/Source/Editor/CurveAssetEditor/Private/SColorGradientCurveEditorView.cpp
--- a/Engine/Source/Editor/CurveAssetEditor/Private/SColorGradientCurveEditorView.cpp
+++ b/Engine/Source/Editor/CurveAssetEditor/Private/SColorGradientCurveEditorView.cpp
@@ -19,7 +19,7 @@ void SColorGradientCurveEditorView::Construct(const FArguments& InArgs, TSharedR
 	bInteractive = 0;
 	bAutoSize = 1;
 	bAllowEmpty = 1;
-	WeakCurveEditor = InCurveEditor;
+	InitCurveEditorReference(InCurveEditor);
 
 	ChildSlot
 	[
diff --git a/Engine/Source/Editor/CurveEditor/Private/CurveEditor.cpp b/Engine/Source/Editor/CurveEditor/Private/CurveEditor.cpp
--- a/Engine/Source/Editor/CurveEditor/Private/CurveEditor.cpp
+++ b/Engine/Source/Editor/CurveEditor/Private/CurveEditor.cpp
@@ -40,6 +40,7 @@
 #include "Misc/FrameNumber.h"
 #include "Misc/FrameTime.h"
 #include "Misc/KeyPasteArgs.h"
+#include "Misc/ScopeExit.h"
 #include "Misc/StringOutputDevice.h"
 #include "Modules/ModuleManager.h"
 #include "SCurveEditor.h" // for access to LogCurveEditor
@@ -241,6 +242,10 @@ FCurveModelID FCurveEditor::AddCurve(TUniquePtr<FCurveModel>&& InCurve)
 		ChildCurves.Add(CurveId, AddCurve(MoveTemp(Child)));
 	}
 
+	// Trigger OnCurveColorsChangedDelegate when any of the colors are changed...
+	// Convenience so external systems do not need to listen to all curves individually.
+	Curve->OnPostColorChanged().AddSP(this, &FCurveEditor::HandleCurveColorChanged, CurveId);
+
 	++ActiveCurvesSerialNumber;
 	if (IsBroadcasting())
 	{
@@ -2178,63 +2183,78 @@ void FCurveEditor::UpdateGeometry(const FGeometry& CurrentGeometry)
 
 void FCurveEditor::SetRandomCurveColorsForSelected()
 {
-	TSet<FCurveModelID> CurveModelIDs = GetSelectionFromTreeAndKeys();
+	const TArray<FCurveModelID> CurveModelIDs = GetSelectionFromTreeAndKeys().Array();
 	if (CurveModelIDs.Num() == 0)
 	{
 		return;
 	}
 
+	// Calling FCurveModel::SetColor below triggers FCurveEditor's own delegate (set up in AddCurve) so skip those and trigger
+	// OnCurveColorsChangedDelegate once at the end.
+	SuspendBroadcast();
+	ON_SCOPE_EXIT{ ResumeBroadcast(); };
+	
 	for (const FCurveModelID& CurveModelID : CurveModelIDs)
 	{
 		if (FCurveModel* Curve = FindCurve(CurveModelID))
 		{
+			const FLinearColor Color = UCurveEditorSettings::GetNextRandomColor();
+			Curve->SetColor(Color);
+			
 			UObject* Object = nullptr;
 			FString Name;
 			Curve->GetCurveColorObjectAndName(&Object, Name);
 			if (Object)
 			{
-				FLinearColor Color = UCurveEditorSettings::GetNextRandomColor();
 				Settings->SetCustomColor(Object->GetClass(), Name, Color);
-				Curve->SetColor(Color);
 			}
 		}
 	}
+
+	OnCurveColorsChangedDelegate.Broadcast(CurveModelIDs);
 }
 
 void FCurveEditor::SetCurveColorsForSelected()
 {
-	TSet<FCurveModelID> CurveModelIDs = GetSelectionFromTreeAndKeys();
+	const TArray<FCurveModelID> CurveModelIDs = GetSelectionFromTreeAndKeys().Array();
 	if (CurveModelIDs.Num() == 0)
 	{
 		return;
 	}
-
-	TWeakPtr<FCurveEditor> WeakSelf = AsShared();
-
+	
 	FColorPickerArgs PickerArgs;
 	PickerArgs.bUseAlpha = false;
-	PickerArgs.InitialColor = FindCurve(*CurveModelIDs.CreateIterator())->GetColor();
-	PickerArgs.OnColorCommitted.BindLambda([WeakSelf, CurveModelIDs](FLinearColor NewColor)
+	PickerArgs.InitialColor = FindCurve(CurveModelIDs[0])->GetColor();
+	PickerArgs.OnColorCommitted.BindLambda([WeakSelf = AsShared().ToWeakPtr(), CurveModelIDs](FLinearColor NewColor)
 	{
-		if (TSharedPtr<FCurveEditor> Self = WeakSelf.Pin())
+		if (const TSharedPtr<FCurveEditor> Self = WeakSelf.Pin())
 		{
 			for (const FCurveModelID& CurveModelID : CurveModelIDs)
 			{
 				if (FCurveModel* Curve = Self->FindCurve(CurveModelID))
 				{
+					Curve->SetColor(NewColor);
+					
 					UObject* Object = nullptr;
 					FString Name;
 					Curve->GetCurveColorObjectAndName(&Object, Name);
 					if (Object)
 					{
 						Self->Settings->SetCustomColor(Object->GetClass(), Name, NewColor);
-						Curve->SetColor(NewColor);
 					}
 				}
+				
+				// This causes a redraw
+				Self->OnCurveColorsChangedDelegate.Broadcast(CurveModelIDs);
 			}
 		}
 	});
-		
+
+	// Calling FCurveModel::SetColor below triggers FCurveEditor's own delegate (set up in AddCurve) so skip those and trigger
+	// OnCurveColorsChangedDelegate once at the end.
+	SuspendBroadcast();
+	ON_SCOPE_EXIT{ ResumeBroadcast(); };
+
 	OpenColorPicker(PickerArgs);
 }
 
@@ -2679,4 +2699,12 @@ void FCurveEditor::OnAxisSnappingChanged()
 	}
 }
 
+void FCurveEditor::HandleCurveColorChanged(FCurveModelID CurveId)
+{
+	if (IsBroadcasting())
+	{
+		OnCurveColorsChangedDelegate.Broadcast(TConstArrayView<FCurveModelID>(&CurveId, 1));
+	}
+}
+
 #undef LOCTEXT_NAMESPACE
diff --git a/Engine/Source/Editor/CurveEditor/Private/CurveEditorCurveDrawParamsCache.cpp b/Engine/Source/Editor/CurveEditor/Private/CurveEditorCurveDrawParamsCache.cpp
--- a/Engine/Source/Editor/CurveEditor/Private/CurveEditorCurveDrawParamsCache.cpp
+++ b/Engine/Source/Editor/CurveEditor/Private/CurveEditorCurveDrawParamsCache.cpp
@@ -19,7 +19,7 @@ namespace UE::CurveEditor
 		GCurveEditorUseCurveCachePool,
 		TEXT("Enables improved curve editor performance by using a curve cache pool"));
 
-	void FCurveDrawParamsCache::Invalidate(const TSharedRef<SCurveEditorView>& CurveEditorView, const TArray<FCurveModelID>& ModelIDs)
+	void FCurveDrawParamsCache::Invalidate(const TSharedRef<SCurveEditorView>& CurveEditorView, TConstArrayView<FCurveModelID> ModelIDs)
 	{
 		TRACE_CPUPROFILER_EVENT_SCOPE(FCurveDrawParamsCache::Invalidate);
 
@@ -305,7 +305,7 @@ namespace UE::CurveEditor
 		}
 	}
 
-	void FCurveDrawParamsCache::DrawCurves(const TSharedRef<SCurveEditorView>& CurveEditorView, const TArray<FCurveModelID>& ModelIDs)
+	void FCurveDrawParamsCache::DrawCurves(const TSharedRef<SCurveEditorView>& CurveEditorView, TConstArrayView<FCurveModelID> ModelIDs)
 	{
 		if (!WeakCurveEditor.IsValid())
 		{
diff --git a/Engine/Source/Editor/CurveEditor/Private/SCurveEditorView.cpp b/Engine/Source/Editor/CurveEditor/Private/SCurveEditorView.cpp
--- a/Engine/Source/Editor/CurveEditor/Private/SCurveEditorView.cpp
+++ b/Engine/Source/Editor/CurveEditor/Private/SCurveEditorView.cpp
@@ -417,6 +417,19 @@ void SCurveEditorView::UpdateCustomAxes()
 	bNeedsDefaultGridLinesV = bHasAnyDefaultAxesV;
 }
 
+void SCurveEditorView::InitCurveEditorReference(const TSharedRef<FCurveEditor>& InCurveEditor)
+{
+	WeakCurveEditor = InCurveEditor;
+	
+	// When the curve colors are changed, issue a redraw.
+	InCurveEditor->OnCurveColorsChangedDelegate.AddSPLambda(this, [this](TConstArrayView<FCurveModelID> Curves)
+	{
+		// The cache does not save any info about colors so it cannot detect them changing... we'll just set the dirty flag and trigger a refresh.
+		CurveDrawParamsCache->CurveCacheFlags = UE::CurveEditor::FCurveDrawParamsCache::ECurveCacheFlags::UpdateAll;
+		RefreshRetainer();
+	});
+}
+
 void SCurveEditorView::FrameVertical(double InOutputMin, double InOutputMax, FCurveEditorViewAxisID AxisID)
 {
 	if (!bFixedOutputBounds && InOutputMin < InOutputMax)
diff --git a/Engine/Source/Editor/CurveEditor/Private/SCurveViewerPanel.cpp b/Engine/Source/Editor/CurveEditor/Private/SCurveViewerPanel.cpp
--- a/Engine/Source/Editor/CurveEditor/Private/SCurveViewerPanel.cpp
+++ b/Engine/Source/Editor/CurveEditor/Private/SCurveViewerPanel.cpp
@@ -35,7 +35,7 @@ namespace CurveViewerConstants
 
 void SCurveViewerPanel::Construct(const FArguments& InArgs, TSharedRef<FCurveEditor> InCurveEditor)
 {
-	WeakCurveEditor = InCurveEditor;
+	InitCurveEditorReference(InCurveEditor);
 	CurveThickness = InArgs._CurveThickness;
 
 	InCurveEditor->SetView(SharedThis(this));
diff --git a/Engine/Source/Editor/CurveEditor/Private/Views/SInteractiveCurveEditorView.cpp b/Engine/Source/Editor/CurveEditor/Private/Views/SInteractiveCurveEditorView.cpp
--- a/Engine/Source/Editor/CurveEditor/Private/Views/SInteractiveCurveEditorView.cpp
+++ b/Engine/Source/Editor/CurveEditor/Private/Views/SInteractiveCurveEditorView.cpp
@@ -122,17 +122,15 @@ public:
 
 void SInteractiveCurveEditorView::Construct(const FArguments& InArgs, TWeakPtr<FCurveEditor> InCurveEditor)
 {
-
 	FixedHeight = InArgs._FixedHeight;
 	BackgroundTint = InArgs._BackgroundTint;
 	MaximumCapacity = InArgs._MaximumCapacity;
 	bAutoSize = InArgs._AutoSize;
-
-	WeakCurveEditor = InCurveEditor;
+	
+	const TSharedPtr<FCurveEditor> CurveEditorPin = InCurveEditor.Pin();
+	check(CurveEditorPin);
+	InitCurveEditorReference(CurveEditorPin.ToSharedRef());
 	{
-		TSharedPtr<FCurveEditor> CurveEditorPin = InCurveEditor.Pin();
-		check(CurveEditorPin);
-
 		CurveDrawParamsCache->Invalidate(SharedThis(this));
 		
 		CurveEditorPin->OnActiveToolChangedDelegate.AddSP(this, &SInteractiveCurveEditorView::OnCurveEditorToolChanged);
diff --git a/Engine/Source/Editor/CurveEditor/Public/CurveEditor.h b/Engine/Source/Editor/CurveEditor/Public/CurveEditor.h
--- a/Engine/Source/Editor/CurveEditor/Public/CurveEditor.h
+++ b/Engine/Source/Editor/CurveEditor/Public/CurveEditor.h
@@ -61,6 +61,7 @@ struct FGeometry;
 DECLARE_DELEGATE_OneParam(FOnSetBoolean, bool)
 DECLARE_MULTICAST_DELEGATE_OneParam(FOnActiveToolChanged, FCurveEditorToolID)
 DECLARE_MULTICAST_DELEGATE_ThreeParams(FOnCurveArrayChanged, FCurveModel*, bool /*displayed*/,const FCurveEditor*);
+DECLARE_MULTICAST_DELEGATE_OneParam(FOnCurveColorsChanged, TConstArrayView<FCurveModelID>); 
 
 /** Enums to describe supported tangent types, by default support all but smart auto*/
 enum class ECurveEditorTangentTypes : int32
@@ -121,6 +122,10 @@ public:
 
 	/** Delegate that is invoked when a tool becomes active. Also fired when the tool goes inactive. */
 	FOnActiveToolChanged OnActiveToolChangedDelegate;
+
+	/** Delegate that is invoked when the curve colors are changed. */
+	FOnCurveColorsChanged OnCurveColorsChangedDelegate;
+	
 public:
 
 	/**
@@ -761,6 +766,9 @@ private:
 	/** Cached physical size of the panel representing this editor */
 	FVector2D CachedPhysicalSize;
 
+	/** Invoked when a curve's color changes (FCurveModel::SetColor is called). */
+	void HandleCurveColorChanged(FCurveModelID CurveId);
+
 public:
 	/**
 	* Delegate that's broadcast when the curve display changes.
diff --git a/Engine/Source/Editor/CurveEditor/Public/CurveEditorCurveDrawParamsCache.h b/Engine/Source/Editor/CurveEditor/Public/CurveEditorCurveDrawParamsCache.h
--- a/Engine/Source/Editor/CurveEditor/Public/CurveEditorCurveDrawParamsCache.h
+++ b/Engine/Source/Editor/CurveEditor/Public/CurveEditorCurveDrawParamsCache.h
@@ -72,7 +72,7 @@ namespace UE::CurveEditor
 		 * Invalidates the cache for specified view and curve model IDs and causes an update on cached data.
 		 * The cache may perform updates that are not instantly relevant asynchronously. To update data synchronous use UpdateCurveDrawParamSynchonous instead.
 		 */
-		UE_API void Invalidate(const TSharedRef<SCurveEditorView>& CurveEditorView, const TArray<FCurveModelID>& ModelIDs = TArray<FCurveModelID>());
+		UE_API void Invalidate(const TSharedRef<SCurveEditorView>& CurveEditorView, TConstArrayView<FCurveModelID> ModelIDs = {});
 
 		/**
 		 * Returns how the curve cache has changed since it was last generated.
@@ -97,7 +97,7 @@ namespace UE::CurveEditor
 		UE_API void UpdateCurveCacheFlags(const TSharedRef<SCurveEditorView>& CurveEditorView);
 
 		/** Draws the curves, possibly async */
-		UE_API void DrawCurves(const TSharedRef<SCurveEditorView>& CurveEditorView, const TArray<FCurveModelID>& ModelIDs = TArray<FCurveModelID>());
+		UE_API void DrawCurves(const TSharedRef<SCurveEditorView>& CurveEditorView, TConstArrayView<FCurveModelID> ModelIDs = {});
 
 		/** Updates curve draw params, possibly async */
 		UE_API void UpdateCurveDrawParams(const TSharedRef<SCurveEditorView>& CurveEditorView, const TArray<FCurveDrawParamsHandle>& CurveDrawParamsHandles);
diff --git a/Engine/Source/Editor/CurveEditor/Public/CurveModel.h b/Engine/Source/Editor/CurveEditor/Public/CurveModel.h
--- a/Engine/Source/Editor/CurveEditor/Public/CurveModel.h
+++ b/Engine/Source/Editor/CurveEditor/Public/CurveModel.h
@@ -445,14 +445,18 @@ public:
 	UE_API TOptional<FKeyHandle> AddKey(const FKeyPosition& NewKeyPosition, const FKeyAttributes& InAttributes);
 
 
-	/**
-	 * Get a multicast delegate, fired when modifications are made to this curve
-	 */
+	/** Get a multicast delegate, fired when modifications are made to this curve */
 	FORCEINLINE FSimpleMulticastDelegate& OnCurveModified()
 	{
 		return CurveModifiedDelegate;
 	}
 
+	/** @return Multicast delegate invoked after the color is changed. */
+	FORCEINLINE FSimpleMulticastDelegate& OnPostColorChanged()
+	{
+		return PostColorChangedDelegate;
+	}
+
 	/**
 	 * Get the owning object of this curve as the specified template type, or search is outer chain if it is not the specified type.
 	 **/
@@ -558,15 +562,10 @@ public:
 		return ChannelName;
 	}
 
-	/**
-	 */
-	FORCEINLINE void SetColor(const FLinearColor& InColor, bool bInModify = true)
+	FORCEINLINE void SetColor(const FLinearColor& InColor)
 	{
 		Color = InColor;
-		if (bInModify)
-		{
-			Modify(); //will make sure the cache get's recreated
-		}
+		PostColorChangedDelegate.Broadcast();
 	}
 
 	/**
@@ -687,6 +686,9 @@ protected:
 	/** Multicast delegate broadcast on curve modification */
 	FSimpleMulticastDelegate CurveModifiedDelegate;
 
+	/** Invokes by SetColor */
+	FSimpleMulticastDelegate PostColorChangedDelegate;
+
 private:
 
 	/**
diff --git a/Engine/Source/Editor/CurveEditor/Public/SCurveEditorView.h b/Engine/Source/Editor/CurveEditor/Public/SCurveEditorView.h
--- a/Engine/Source/Editor/CurveEditor/Public/SCurveEditorView.h
+++ b/Engine/Source/Editor/CurveEditor/Public/SCurveEditorView.h
@@ -342,6 +342,10 @@ public:
 	UE_API void UpdateCustomAxes();
 
 protected:
+
+	/** Inits this view with the curve editor and sets up required callbacks. */
+	UE_API void InitCurveEditorReference(const TSharedRef<FCurveEditor>& InCurveEditor);
+	
 	/** Gets info about the curves being drawn. Converts actual curves into an abstract series of lines/points/handles/etc. */
 	UE_DEPRECATED(5.6, "Instead use the CurveDrawParamsCache, use FCurveDrawParamsCache::GetCurveDrawParamsSynchronous to get Curve Draw Parms as with this function.")
 	UE_API void GetCurveDrawParams(TArray<FCurveDrawParams>& OutDrawParams);
diff --git a/Engine/Source/Editor/MovieSceneTools/Private/Channels/SCurveEditorEventChannelView.cpp b/Engine/Source/Editor/MovieSceneTools/Private/Channels/SCurveEditorEventChannelView.cpp
--- a/Engine/Source/Editor/MovieSceneTools/Private/Channels/SCurveEditorEventChannelView.cpp
+++ b/Engine/Source/Editor/MovieSceneTools/Private/Channels/SCurveEditorEventChannelView.cpp
@@ -37,7 +37,6 @@ void SCurveEditorEventChannelView::Construct(const FArguments& InArgs, TWeakPtr<
 	bAllowModelViewTransforms = false;
 	OutputMin = -0.5;
 	OutputMax =  0.5;
-	WeakCurveEditor = InCurveEditor;
 	SortBias = 25;
 
 	SInteractiveCurveEditorView::Construct(InArgs, InCurveEditor);
diff --git a/Engine/Source/Editor/MovieSceneTools/Private/Channels/SCurveEditorKeyBarView.cpp b/Engine/Source/Editor/MovieSceneTools/Private/Channels/SCurveEditorKeyBarView.cpp
--- a/Engine/Source/Editor/MovieSceneTools/Private/Channels/SCurveEditorKeyBarView.cpp
+++ b/Engine/Source/Editor/MovieSceneTools/Private/Channels/SCurveEditorKeyBarView.cpp
@@ -44,7 +44,6 @@ void SCurveEditorKeyBarView::Construct(const FArguments& InArgs, TWeakPtr<FCurve
 
 	OutputMin = -0.5;
 	OutputMax =  0.5;
-	WeakCurveEditor = InCurveEditor;
 	SortBias = 25;
 
 	SInteractiveCurveEditorView::Construct(InArgs, InCurveEditor);
diff --git a/Engine/Source/Editor/Persona/Private/AnimTimeline/AnimTimelineTrack_Curve.cpp b/Engine/Source/Editor/Persona/Private/AnimTimeline/AnimTimelineTrack_Curve.cpp
--- a/Engine/Source/Editor/Persona/Private/AnimTimeline/AnimTimelineTrack_Curve.cpp
+++ b/Engine/Source/Editor/Persona/Private/AnimTimeline/AnimTimelineTrack_Curve.cpp
@@ -216,7 +216,7 @@ TSharedRef<SWidget> FAnimTimelineTrack_Curve::GenerateContainerWidgetForTimeline
 		GetCurveEditInfo(CurveIndex, Name, Type, EditIndex);
 
 		TUniquePtr<FRichCurveEditorModelNamed> NewCurveModel = MakeUnique<FRichCurveEditorModelNamed>(Name, Type, EditIndex, GetModel()->GetAnimSequenceBase());
-		NewCurveModel->SetColor(GetCurveColor(CurveIndex), false);
+		NewCurveModel->SetColor(GetCurveColor(CurveIndex));
 		NewCurveModel->SetIsKeyDrawEnabled(MakeAttributeLambda([](){ return GetDefault<UPersonaOptions>()->bTimelineDisplayCurveKeys; }));
 		CurveEditor->AddCurve(MoveTemp(NewCurveModel));
 	}
diff --git a/Engine/Source/Editor/Persona/Private/AnimTimeline/AnimTimelineTrack_FloatCurve.cpp b/Engine/Source/Editor/Persona/Private/AnimTimeline/AnimTimelineTrack_FloatCurve.cpp
--- a/Engine/Source/Editor/Persona/Private/AnimTimeline/AnimTimelineTrack_FloatCurve.cpp
+++ b/Engine/Source/Editor/Persona/Private/AnimTimeline/AnimTimelineTrack_FloatCurve.cpp
@@ -349,7 +349,7 @@ void FAnimTimelineTrack_FloatCurve::AddCurveTrackButton(TSharedPtr<SHorizontalBo
 		// Set display curves too
 		for(const TPair<FCurveModelID, TUniquePtr<FCurveModel>>& CurvePair : CurveEditor->GetCurves())
 		{
-			CurvePair.Value->SetColor(InNewColor, false);
+			CurvePair.Value->SetColor(InNewColor);
 		}
 	};
 
diff --git a/Engine/Source/Editor/Sequencer/Private/IKeyArea.cpp b/Engine/Source/Editor/Sequencer/Private/IKeyArea.cpp
--- a/Engine/Source/Editor/Sequencer/Private/IKeyArea.cpp
+++ b/Engine/Source/Editor/Sequencer/Private/IKeyArea.cpp
@@ -473,7 +473,7 @@ TUniquePtr<FCurveModel> IKeyArea::CreateCurveEditorModel(TSharedRef<ISequencer>
 			CurveModel->SetChannelName(MetaData->Name);
 			if (Color.IsSet())
 			{
-				CurveModel->SetColor(Color.GetValue(),false);
+				CurveModel->SetColor(Color.GetValue());
 			}
 
 			//Use editor preference color if it's been set, this function is const so we can't set just Color optional
@@ -486,7 +486,7 @@ TUniquePtr<FCurveModel> IKeyArea::CreateCurveEditorModel(TSharedRef<ISequencer>
 				TOptional<FLinearColor> SettingColor = Settings->GetCustomColor(Object->GetClass(), Name);
 				if (SettingColor.IsSet())
 				{
-					CurveModel->SetColor(SettingColor.GetValue(),false);
+					CurveModel->SetColor(SettingColor.GetValue());
 				}
 			}
 
