Curve Editor: Introduce "CurveEditor.UseNewUndoSystem" CVar, which is disabled by default. This should more or less revert the behaviour to be same as before the undo refactors, giving us time to address any bugs that the new system causes.

diff --git a/Engine/Source/Editor/CurveEditor/Private/Modification/Utils/ScopedCurveChange.cpp b/Engine/Source/Editor/CurveEditor/Private/Modification/Utils/ScopedCurveChange.cpp
--- a/Engine/Source/Editor/CurveEditor/Private/Modification/Utils/ScopedCurveChange.cpp
+++ b/Engine/Source/Editor/CurveEditor/Private/Modification/Utils/ScopedCurveChange.cpp
@@ -23,6 +23,21 @@ static TAutoConsoleVariable<bool> CVarLogUndoKeyChangeSize(
 	TEXT("Whenever a key change is made, whether to log debug info (time spent diffing, bytes used for the diff, the final bytes on undo stack size).")
 	);
 
+static TAutoConsoleVariable<bool> CVarUseNewUndoSystem(
+	TEXT("CurveEditor.UseNewUndoSystem"), false, TEXT("Whether to use the new undo system. Warning: There may still be some bugs.")
+	);
+
+static void ModifyCurves(const TSharedPtr<FCurveEditor>& InCurveEditorPin, const FCurveChangeDiff& InChangeDiff)
+{
+	for (const TPair<FCurveModelID, FCurveDiffingData>& Pair : InChangeDiff.GetSnapshot().CurveData)
+	{
+		if (FCurveModel* CurveModel = InCurveEditorPin->FindCurve(Pair.Key))
+		{
+			CurveModel->Modify();
+		}
+	}
+}
+
 FScopedCurveChange::FScopedCurveChange(
 	TWeakPtr<FCurveEditor> InCurveEditor, FCurveChangeDiff InKeyChangeDiff, const FText& InDescription, EScopedKeyChangeFlags InFlags
 	)
@@ -30,7 +45,14 @@ FScopedCurveChange::FScopedCurveChange(
 	, Transaction(Description)
 	, ChangeDiff(MoveTemp(InKeyChangeDiff))
 	, Flags(InFlags)
-{}
+{
+	const TSharedPtr<FCurveEditor> CurveEditorPin = WeakCurveEditor.Pin();
+	if (CurveEditorPin && !CVarUseNewUndoSystem.GetValueOnAnyThread())
+	{
+		// If the new undo system is disabled, just do what we did before: call Modify().
+		ModifyCurves(CurveEditorPin, ChangeDiff);
+	}
+}
 
 FScopedCurveChange::FScopedCurveChange(
 	FCurvesSnapshotBuilder InDataToDiff, const FText& InDescription, EScopedKeyChangeFlags InFlags
@@ -44,6 +66,11 @@ FScopedCurveChange::FScopedCurveChange(
 
 FScopedCurveChange::~FScopedCurveChange()
 {
+	if (!CVarUseNewUndoSystem.GetValueOnAnyThread())
+	{
+		return;
+	}
+	
 	const TSharedPtr<FCurveEditor> CurveEditorPin = WeakCurveEditor.Pin();
 	if (!CurveEditorPin)
 	{
