Implement GetSize in FGenericCurveChangeCommand and FSelectionChangeCommand.

diff --git a/Engine/Source/Editor/CurveEditor/Private/Modification/Keys/GenericCurveChangeCommand.cpp b/Engine/Source/Editor/CurveEditor/Private/Modification/Keys/GenericCurveChangeCommand.cpp
--- a/Engine/Source/Editor/CurveEditor/Private/Modification/Keys/GenericCurveChangeCommand.cpp
+++ b/Engine/Source/Editor/CurveEditor/Private/Modification/Keys/GenericCurveChangeCommand.cpp
@@ -22,4 +22,9 @@ void FGenericCurveChangeCommand::Revert(UObject* Object)
 		GenericCurveChange::RevertChange(CurveEditor->GetCurves(), DeltaChange);
 	}
 }
+
+SIZE_T FGenericCurveChangeCommand::GetSize() const
+{
+	return DeltaChange.GetAllocatedSize() + sizeof(FGenericCurveChangeCommand);
+}
 }
diff --git a/Engine/Source/Editor/CurveEditor/Private/Modification/Keys/GenericCurveChangeCommand.h b/Engine/Source/Editor/CurveEditor/Private/Modification/Keys/GenericCurveChangeCommand.h
--- a/Engine/Source/Editor/CurveEditor/Private/Modification/Keys/GenericCurveChangeCommand.h
+++ b/Engine/Source/Editor/CurveEditor/Private/Modification/Keys/GenericCurveChangeCommand.h
@@ -24,6 +24,7 @@ public:
 	//~ Begin FChange Interface
 	virtual void Apply(UObject* Object) override;
 	virtual void Revert(UObject* Object) override;
+	virtual SIZE_T GetSize() const override;
 	//~ End FChange Interface
 
 private:
diff --git a/Engine/Source/Editor/CurveEditor/Private/Modification/Selection/SelectionChangeCommand.cpp b/Engine/Source/Editor/CurveEditor/Private/Modification/Selection/SelectionChangeCommand.cpp
--- a/Engine/Source/Editor/CurveEditor/Private/Modification/Selection/SelectionChangeCommand.cpp
+++ b/Engine/Source/Editor/CurveEditor/Private/Modification/Selection/SelectionChangeCommand.cpp
@@ -22,4 +22,9 @@ void FSelectionChangeCommand::Revert(UObject* Object)
 		RevertSelectionChange(CurveEditor->Selection, DeltaChange);
 	}
 }
+
+SIZE_T FSelectionChangeCommand::GetSize() const
+{
+	return sizeof(FSelectionChangeCommand) + DeltaChange.GetAllocatedSize();
+}
 }
diff --git a/Engine/Source/Editor/CurveEditor/Private/Modification/Selection/SelectionChangeCommand.h b/Engine/Source/Editor/CurveEditor/Private/Modification/Selection/SelectionChangeCommand.h
--- a/Engine/Source/Editor/CurveEditor/Private/Modification/Selection/SelectionChangeCommand.h
+++ b/Engine/Source/Editor/CurveEditor/Private/Modification/Selection/SelectionChangeCommand.h
@@ -17,6 +17,7 @@ public:
 	//~ Begin FChange Interface
 	virtual void Apply(UObject* Object) override;
 	virtual void Revert(UObject* Object) override;
+	virtual SIZE_T GetSize() const override;
 	//~ End FChange Interface
 
 private:
diff --git a/Engine/Source/Editor/CurveEditor/Private/Modification/Selection/SelectionDeltaChange.h b/Engine/Source/Editor/CurveEditor/Private/Modification/Selection/SelectionDeltaChange.h
--- a/Engine/Source/Editor/CurveEditor/Private/Modification/Selection/SelectionDeltaChange.h
+++ b/Engine/Source/Editor/CurveEditor/Private/Modification/Selection/SelectionDeltaChange.h
@@ -2,6 +2,7 @@
 
 #pragma once
 
+#include "Algo/Accumulate.h"
 #include "Containers/Array.h"
 #include "Containers/Map.h"
 #include "CurveEditorTypes.h"
@@ -21,6 +22,11 @@ struct FCurveSelectionDeltaChange
 	 * Can contain a key handle that is in AddedKeys. In that case, the point type was changed.
 	 */
 	TArray<TPair<FKeyHandle, ECurvePointType>> RemovedKeys;
+
+	SIZE_T GetAllocatedSize() const
+	{
+		return AddedKeys.GetAllocatedSize() + RemovedKeys.GetAllocatedSize();
+	}
 };
 
 /** Changes made to the curve editor selection */
@@ -36,5 +42,14 @@ struct FSelectionDeltaChange
 
 	bool HasChanges() const { return !ChangedCurves.IsEmpty() || OldSerialNumber != NewSerialNumber; }
 	operator bool() const { return HasChanges(); }
+
+	SIZE_T GetAllocatedSize() const
+	{
+		return ChangedCurves.GetAllocatedSize()
+			+ Algo::TransformAccumulate(ChangedCurves, [](const TPair<FCurveModelID, FCurveSelectionDeltaChange>& Pair)
+			{
+				return Pair.Value.GetAllocatedSize();
+			}, 0);
+	}
 };
 }
