Curve Editor: Fix level sequence not getting dirtied after undo & redo.

The idea of the fix is to call GUndo->StoreUndo for each UMovieSceneSection object and pass a mini FGenericCurveChangeCommand.
That means each FCurveModel will have its own FGenericCurveChangeCommand.
For FCurveModels that return nullptr from GetOwningObject, we continue to build a big FGenericCurveChangeCommand that contains all curve changes.

diff --git a/Engine/Source/Editor/CurveEditor/Private/Modification/CurveEditorCommandChange.cpp b/Engine/Source/Editor/CurveEditor/Private/Modification/CurveEditorCommandChange.cpp
--- a/Engine/Source/Editor/CurveEditor/Private/Modification/CurveEditorCommandChange.cpp
+++ b/Engine/Source/Editor/CurveEditor/Private/Modification/CurveEditorCommandChange.cpp
@@ -9,7 +9,7 @@ namespace UE::CurveEditor
 {
 bool FCurveEditorCommandChange::HasExpired(UObject* InObject) const
 {
-	return !GetCurveEditor(InObject).IsValid();
+	return !GetCurveEditor().IsValid();
 }
 
 FString FCurveEditorCommandChange::ToString() const
@@ -17,9 +17,8 @@ FString FCurveEditorCommandChange::ToString() const
 	return TEXT("FSelectionChangeCommand");
 }
 
-TSharedPtr<FCurveEditor> FCurveEditorCommandChange::GetCurveEditor(UObject* InObject)
+TSharedPtr<FCurveEditor> FCurveEditorCommandChange::GetCurveEditor() const
 {
-	UCurveEditorTransactionObject* TransactionObject = Cast<UCurveEditorTransactionObject>(InObject);
-	return ensure(TransactionObject) ? TransactionObject->OwningCurveEditor.Pin() : nullptr;
+	return WeakCurveEditor.Pin();
 }
 }
diff --git a/Engine/Source/Editor/CurveEditor/Private/Modification/Keys/CurveChangeDiff.cpp b/Engine/Source/Editor/CurveEditor/Private/Modification/Keys/CurveChangeDiff.cpp
--- a/Engine/Source/Editor/CurveEditor/Private/Modification/Keys/CurveChangeDiff.cpp
+++ b/Engine/Source/Editor/CurveEditor/Private/Modification/Keys/CurveChangeDiff.cpp
@@ -8,19 +8,23 @@
 #include "CurveModel.h"
 #include "KeyAttributeChangeUtils.h"
 #include "KeyExistenceUtils.h"
+#include "Misc/ScopeExit.h"
 #include "Modification/Keys/Data/CurveKeyDataMap.h"
 #include "Modification/Keys/Data/GenericCurveChangeData.h"
 #include "Modification/Keys/ECurveChangeFlags.h"
 #include "Modification/Keys/MoveKeysChangeUtils.h"
 #include "SCurveEditor.h"
+#include "Diff/AllCurvesGenericChangeBuilder.h"
+#include "Modification/Keys/Diff/IMultiCurveChangeVisitor.h"
+#include "Modification/Keys/Diff/ISingleCurveChangeVisitor.h"
 
 namespace UE::CurveEditor
 {
 namespace Private
 {
 static void HandleOp_MoveKeys(
-	const ECurveChangeFlags Operations, const FCurveModelID& CurveId, const FCurveModel& CurveModel, const FCurveKeyData& SavedState,
-	FGenericCurveChangeData& OutChange
+	const ECurveChangeFlags Operations, const FCurveModel& CurveModel, const FCurveKeyData& SavedState,
+	ISingleCurveChangeVisitor& InCurveChangeVisitor
 	)
 {
 	if (!EnumHasAnyFlags(Operations, ECurveChangeFlags::MoveKeys) || !SavedState || !ensureAlways(SavedState.HasKeyPositions()))
@@ -34,13 +38,13 @@ static void HandleOp_MoveKeys(
 			
 	if (FMoveKeysChangeData_PerCurve PerCurveData = MoveKeys::Diff(SavedState.KeyHandles, SavedState.KeyPositions, CurrentPositions))
 	{
-		OutChange.MoveKeysData.ChangedCurves.Emplace(CurveId, MoveTemp(PerCurveData));
+		InCurveChangeVisitor.ProcessMoveKeys(MoveTemp(PerCurveData));
 	}
 }
 
 static void HandleOp_AddedKeys(
-	const ECurveChangeFlags Operations, const FCurveModelID& CurveId, const FCurveModel& CurveModel, const FCurveKeyData& SavedState,
-	FGenericCurveChangeData& OutChange
+	const ECurveChangeFlags Operations, const FCurveModel& CurveModel, const FCurveKeyData& SavedState,
+	ISingleCurveChangeVisitor& InCurveChangeVisitor
 )
 {
 	if (!EnumHasAnyFlags(Operations, ECurveChangeFlags::AddKeys) || !SavedState)
@@ -50,30 +54,25 @@ static void HandleOp_AddedKeys(
 
 	if (FCurveKeyData AddedKeys = KeyInsertion::Diff(SavedState.KeyHandles, CurveModel))
 	{
-		OutChange.AddKeysData.SavedCurveState.Emplace(CurveId, MoveTemp(AddedKeys));
+		InCurveChangeVisitor.ProcessAddKeys(MoveTemp(AddedKeys));
 	}
 }
 
-static void HandleOp_RemovedKeys(
-	const ECurveChangeFlags Operations, const FCurveModelID& CurveId, const FCurveModel& CurveModel, const FCurveKeyData& SavedState,
-	FGenericCurveChangeData& OutChange
-)
+static TOptional<FCurveKeyData> ComputeRemovedKeys(const ECurveChangeFlags Operations, const FCurveModel& CurveModel, const FCurveKeyData& SavedState)
 {
 	if (!EnumHasAnyFlags(Operations, ECurveChangeFlags::RemoveKeys) || !SavedState
 		|| !ensureAlways(SavedState.HasKeyPositions()) || !ensureAlways(SavedState.HasKeyAttributes()))
 	{
-		return;
+		return {};
 	}
 
-	if (FCurveKeyData Removed = KeyRemoval::Diff(SavedState.KeyHandles, SavedState.KeyPositions, SavedState.KeyAttributes, CurveModel))
-	{
-		OutChange.RemoveKeysData.SavedCurveState.Emplace(CurveId, MoveTemp(Removed));
-	}
+	FCurveKeyData Removed = KeyRemoval::Diff(SavedState.KeyHandles, SavedState.KeyPositions, SavedState.KeyAttributes, CurveModel);
+	return Removed ? Removed : TOptional<FCurveKeyData>();
 }
 	
 static void HandleOp_KeyAttributes(
-	const ECurveChangeFlags Operations, const FCurveModelID& CurveId, const FCurveModel& CurveModel, const FCurveKeyData& SavedState,
-	FGenericCurveChangeData& OutChange
+	const ECurveChangeFlags Operations, const FCurveModel& CurveModel, const FCurveKeyData& SavedState,
+	ISingleCurveChangeVisitor& InCurveChangeVisitor
 )
 {
 	if (!EnumHasAnyFlags(Operations, ECurveChangeFlags::KeyAttributes) || !SavedState || !ensureAlways(SavedState.HasKeyAttributes()))
@@ -87,13 +86,13 @@ static void HandleOp_KeyAttributes(
 
 	if (FKeyAttributeChangeData_PerCurve PerCurveData = KeyAttributes::Diff(SavedState.KeyHandles, SavedState.KeyAttributes, CurrentAttributes))
 	{
-		OutChange.KeyAttributeData.ChangedCurves.Emplace(CurveId, MoveTemp(PerCurveData));
+		InCurveChangeVisitor.ProcessKeyAttributesChange(MoveTemp(PerCurveData));
 	}
 }
 
 static void HandleOp_CurveAttributes(
-	const ECurveChangeFlags Operations, const FCurveModelID& CurveId, const FCurveModel& CurveModel, const FCurveDiffingData& SavedState,
-	FGenericCurveChangeData& OutChange
+	const ECurveChangeFlags Operations, const FCurveModel& CurveModel, const FCurveDiffingData& SavedState,
+	ISingleCurveChangeVisitor& InCurveChangeVisitor
 )
 {
 	if (!EnumHasAnyFlags(Operations, ECurveChangeFlags::CurveAttributes) || !SavedState || !ensureAlways(SavedState.HasCurveAttributes()))
@@ -106,7 +105,7 @@ static void HandleOp_CurveAttributes(
 	
 	if (FCurveAttributeChangeData_PerCurve PerCurveData = CurveAttributes::Diff(*SavedState.CurveAttributes, CurrentAttributes))
 	{
-		OutChange.CurveAttributeData.ChangeData.Emplace(CurveId, PerCurveData);
+		InCurveChangeVisitor.ProcessCurveAttributesChange(MoveTemp(PerCurveData));
 	}
 }
 
@@ -157,14 +156,14 @@ static FCurveKeyData CleanseOfRemovedKeys(
 }
 
 /** Checks whether any keys in InSavedState were removed from CurveModel and if so, returns a copy of the state minus the removed key data. */
-TOptional<FCurveKeyData> CleanseStateIfKeysWereRemoved(
-	const FCurveKeyData& InSavedState, const ECurveChangeFlags Operations, const FRemoveKeysChangeData& InKeyRemovalData,
-	const FCurveModelID& CurveId, const FCurveModel& CurveModel
+static TOptional<FCurveKeyData> CleanseStateIfKeysWereRemoved(
+	const FCurveKeyData& InSavedState, const ECurveChangeFlags Operations, const TOptional<FCurveKeyData>& InComputedKeyRemovalData,
+	const FCurveModel& CurveModel
 	)
 {
-	if (const FCurveKeyData* RemovedKeyData = InKeyRemovalData.SavedCurveState.Find(CurveId))
+	if (InComputedKeyRemovalData)
 	{
-		return CleanseOfRemovedKeys(RemovedKeyData->KeyHandles, InSavedState);
+		return CleanseOfRemovedKeys(InComputedKeyRemovalData->KeyHandles, InSavedState);
 	}
 
 	// If InKeyRemovalData contained no removal data but RemoveKeys is set, that means we already analyzed removal and don't have to do it again.
@@ -178,29 +177,53 @@ TOptional<FCurveKeyData> CleanseStateIfKeysWereRemoved(
 	return {};
 }
 }
-
+	
 FGenericCurveChangeData ComputeDiff(const FCurveEditor& InCurveEditor, const ECurveChangeFlags InOperations, const FCurvesSnapshot& InSnapshot)
+{
+	FAllCurvesGenericChangeBuilder Builder(InCurveEditor);
+	ProcessChanges(InCurveEditor, InOperations, InSnapshot, Builder);
+	return Builder.CurveChangeData;
+}
+	
+void ProcessChanges(
+	const FCurveEditor& InCurveEditor, const ECurveChangeFlags InOperations, const FCurvesSnapshot& InSnapshot,
+	IMultiCurveChangeVisitor& InChangeVisitor
+	)
 {
 	SCOPED_CURVE_EDITOR_TRACE(ComputeDiff);
 	
+	InChangeVisitor.PreProcessChanges();
+	ON_SCOPE_EXIT{ InChangeVisitor.PostProcessChanges(); };
+	
 	FGenericCurveChangeData Change;
 	for (const TPair<FCurveModelID, FCurveDiffingData>& StatePair : InSnapshot.CurveData)
 	{
 		const FCurveModelID& CurveId = StatePair.Key;
-		if (const FCurveModel* CurveModel = InCurveEditor.FindCurve(CurveId))
+		const FCurveDiffingData& SavedState = StatePair.Value;
+		InChangeVisitor.ProcessChange(CurveId, [&InCurveEditor, InOperations, &CurveId, &SavedState]
+			(ISingleCurveChangeVisitor& InSingleCurveChangeVisitor)
 		{
-			const FCurveDiffingData& SavedState = StatePair.Value;
+			const FCurveModel* CurveModel = InCurveEditor.FindCurve(CurveId);
+			if (!CurveModel)
+			{
+				return;
+			}
 
-			Private::HandleOp_CurveAttributes(InOperations, CurveId, *CurveModel, SavedState, Change);
-			
-			// Do RemoveKeys before MoveKeys and KeyAttributes... 
-			Private::HandleOp_RemovedKeys(InOperations, CurveId, *CurveModel, SavedState, Change);
-			Private::HandleOp_AddedKeys(InOperations, CurveId, *CurveModel, SavedState, Change);
+			Private::HandleOp_CurveAttributes(InOperations, *CurveModel, SavedState, InSingleCurveChangeVisitor);
+			Private::HandleOp_AddedKeys(InOperations, *CurveModel, SavedState, InSingleCurveChangeVisitor);
 			
+			// Compute removed keys before MoveKeys and KeyAttributes... 
+			TOptional<FCurveKeyData> RemovedKeys = Private::ComputeRemovedKeys(InOperations, *CurveModel, SavedState);
 			// ... so we can reuse that info to tell MoveKeys and KeyAttributes to skip analyzing the keys that were removed.
 			const TOptional<FCurveKeyData> StateMinusRemovedKeys = Private::CleanseStateIfKeysWereRemoved(
-				SavedState, InOperations, Change.RemoveKeysData, CurveId, *CurveModel
+				SavedState, InOperations, RemovedKeys, *CurveModel
 				);
+			if (RemovedKeys)
+			{
+				// Now that CleanseStateIfKeysWereRemoved has used RemovedKeys, we can MoveTemp it and pass it to the visitor.
+				InSingleCurveChangeVisitor.ProcessRemoveKeys(MoveTemp(*RemovedKeys));
+			}
+				
 			if (StateMinusRemovedKeys)
 			{
 				// The operations should ignore the keys that were removed since that data is now irrelevant.
@@ -209,18 +232,31 @@ FGenericCurveChangeData ComputeDiff(const FCurveEditor& InCurveEditor, const ECu
 				// that were removed by the operation.
 				// GetKeyPositions would return default-initialized FKeyPositions. It would then proceed to compare the default-initialized state
 				// against the key value from before the change, and conclude that the keys removed by the operation were "moved" (false change).
-				Private::HandleOp_MoveKeys(InOperations, CurveId, *CurveModel, *StateMinusRemovedKeys, Change);
-				Private::HandleOp_KeyAttributes(InOperations, CurveId, *CurveModel, *StateMinusRemovedKeys, Change);
+				Private::HandleOp_MoveKeys(InOperations, *CurveModel, *StateMinusRemovedKeys, InSingleCurveChangeVisitor);
+				Private::HandleOp_KeyAttributes(InOperations, *CurveModel, *StateMinusRemovedKeys, InSingleCurveChangeVisitor);
 			}
 			else
 			{
 				// Nothing was removed, so just pass in the same state
-				Private::HandleOp_MoveKeys(InOperations, CurveId, *CurveModel, SavedState, Change);
-				Private::HandleOp_KeyAttributes(InOperations, CurveId, *CurveModel, SavedState, Change);
+				Private::HandleOp_MoveKeys(InOperations, *CurveModel, SavedState, InSingleCurveChangeVisitor);
+				Private::HandleOp_KeyAttributes(InOperations, *CurveModel, SavedState, InSingleCurveChangeVisitor);
 			}
-		}
+		});
 	}
+}
+
+FGenericCurveChangeData FCurveChangeDiff::ComputeDiff() const
+{
+	const TSharedPtr<FCurveEditor> CurveEditorPin = WeakCurveEditor.Pin();
+	return CurveEditorPin ? CurveEditor::ComputeDiff(*CurveEditorPin, Operations, Snapshot) : FGenericCurveChangeData();
 	
-	return Change;
 }
-}
\ No newline at end of file
+
+void FCurveChangeDiff::ProcessDiffs(IMultiCurveChangeVisitor& InChangeBuilder) const
+{
+	if (const TSharedPtr<FCurveEditor> CurveEditorPin = WeakCurveEditor.Pin())
+	{
+		ProcessChanges(*CurveEditorPin, Operations, Snapshot, InChangeBuilder);
+	}
+}
+}
diff --git a/Engine/Source/Editor/CurveEditor/Private/Modification/Keys/Diff/AllCurvesGenericChangeBuilder.cpp b/Engine/Source/Editor/CurveEditor/Private/Modification/Keys/Diff/AllCurvesGenericChangeBuilder.cpp
new file mode 100644
--- /dev/null
+++ b/Engine/Source/Editor/CurveEditor/Private/Modification/Keys/Diff/AllCurvesGenericChangeBuilder.cpp
@@ -0,0 +1,17 @@
+﻿// Copyright Epic Games, Inc. All Rights Reserved.
+
+#include "AllCurvesGenericChangeBuilder.h"
+
+#include "AppendPerCurveGenericChangeVisitor.h"
+#include "Modification/Keys/Diff/ISingleCurveChangeVisitor.h"
+
+namespace UE::CurveEditor
+{
+void FAllCurvesGenericChangeBuilder::ProcessChange(
+	const FCurveModelID& InCurveModel, TFunctionRef<void(ISingleCurveChangeVisitor& InVisitor)> InProcessCallback
+	)
+{
+	FAppendPerCurveGenericChangeVisitor Visitor(InCurveModel, CurveChangeData);
+	InProcessCallback(Visitor);
+}
+}
diff --git a/Engine/Source/Editor/CurveEditor/Private/Modification/Keys/Diff/AllCurvesGenericChangeBuilder.h b/Engine/Source/Editor/CurveEditor/Private/Modification/Keys/Diff/AllCurvesGenericChangeBuilder.h
new file mode 100644
--- /dev/null
+++ b/Engine/Source/Editor/CurveEditor/Private/Modification/Keys/Diff/AllCurvesGenericChangeBuilder.h
@@ -0,0 +1,32 @@
+﻿// Copyright Epic Games, Inc. All Rights Reserved.
+
+#pragma once
+
+#include "Modification/Keys/Data/GenericCurveChangeData.h"
+#include "Modification/Keys/Diff/IMultiCurveChangeVisitor.h"
+
+class FCurveEditor;
+
+namespace UE::CurveEditor
+{
+/** Aggregates all curve changes into a single FGenericCurveChangeData. */
+class FAllCurvesGenericChangeBuilder : public IMultiCurveChangeVisitor
+{
+public:
+
+	FGenericCurveChangeData CurveChangeData;
+
+	explicit FAllCurvesGenericChangeBuilder(const FCurveEditor& InCurveEditor UE_LIFETIMEBOUND) : CurveEditor(InCurveEditor) {}
+	
+	//~ Begin ICurvesDiffer Interface
+	virtual void ProcessChange(const FCurveModelID& InCurveModel, TFunctionRef<void(ISingleCurveChangeVisitor&)> InProcessCallback) override;
+	//~ End ICurvesDiffer Interface
+
+private:
+
+	/** Curve editor instance. Caller is responsible for keeping it alive. */
+	const FCurveEditor& CurveEditor;
+};
+}
+
+
diff --git a/Engine/Source/Editor/CurveEditor/Private/Modification/Keys/Diff/AppendPerCurveGenericChangeVisitor.h b/Engine/Source/Editor/CurveEditor/Private/Modification/Keys/Diff/AppendPerCurveGenericChangeVisitor.h
new file mode 100644
--- /dev/null
+++ b/Engine/Source/Editor/CurveEditor/Private/Modification/Keys/Diff/AppendPerCurveGenericChangeVisitor.h
@@ -0,0 +1,47 @@
+﻿// Copyright Epic Games, Inc. All Rights Reserved.
+
+#pragma once
+
+#include "CurveEditorTypes.h"
+#include "Modification/Keys/Data/GenericCurveChangeData.h"
+#include "Modification/Keys/Diff/ISingleCurveChangeVisitor.h"
+
+namespace UE::CurveEditor
+{
+class FAppendPerCurveGenericChangeVisitor : public ISingleCurveChangeVisitor
+{
+	const FCurveModelID CurveId;
+	FGenericCurveChangeData& CurveChangeData;
+public:
+
+	explicit FAppendPerCurveGenericChangeVisitor(const FCurveModelID InCurveId, FGenericCurveChangeData& InCurveChangeData UE_LIFETIMEBOUND)
+		: CurveId(InCurveId)
+		, CurveChangeData(InCurveChangeData)
+	{}
+
+	virtual void ProcessMoveKeys(FMoveKeysChangeData_PerCurve&& InData) override
+	{
+		CurveChangeData.MoveKeysData.ChangedCurves.Add(CurveId, MoveTemp(InData));
+	}
+
+	virtual void ProcessAddKeys(FCurveKeyData&& InData) override
+	{
+		CurveChangeData.AddKeysData.SavedCurveState.Add(CurveId, MoveTemp(InData));
+	}
+
+	virtual void ProcessRemoveKeys(FCurveKeyData&& InData) override
+	{
+		CurveChangeData.RemoveKeysData.SavedCurveState.Add(CurveId, MoveTemp(InData));
+	}
+
+	virtual void ProcessKeyAttributesChange(FKeyAttributeChangeData_PerCurve&& InData) override
+	{
+		CurveChangeData.KeyAttributeData.ChangedCurves.Add(CurveId, MoveTemp(InData));
+	}
+
+	virtual void ProcessCurveAttributesChange(FCurveAttributeChangeData_PerCurve&& InData) override
+	{
+		CurveChangeData.CurveAttributeData.ChangeData.Add(CurveId, MoveTemp(InData));
+	}
+};
+}
diff --git a/Engine/Source/Editor/CurveEditor/Private/Modification/Keys/Diff/MultiGenericCommandAppender.cpp b/Engine/Source/Editor/CurveEditor/Private/Modification/Keys/Diff/MultiGenericCommandAppender.cpp
new file mode 100644
--- /dev/null
+++ b/Engine/Source/Editor/CurveEditor/Private/Modification/Keys/Diff/MultiGenericCommandAppender.cpp
@@ -0,0 +1,85 @@
+﻿// Copyright Epic Games, Inc. All Rights Reserved.
+
+#include "MultiGenericCommandAppender.h"
+
+#include "CurveEditor.h"
+#include "AppendPerCurveGenericChangeVisitor.h"
+#include "Modification/Keys/GenericCurveChangeCommand.h"
+
+namespace UE::CurveEditor
+{
+namespace Private
+{
+static void AppendChanges(const FCurveModelID& InCurveId, FGenericCurveChangeData& InOutTarget, FGenericCurveChangeData&& InAppendedChanges)
+{
+	if (FMoveKeysChangeData_PerCurve* Change_Move = InAppendedChanges.MoveKeysData.ChangedCurves.Find(InCurveId))
+	{
+		InOutTarget.MoveKeysData.ChangedCurves.Add(InCurveId, MoveTemp(*Change_Move));
+	}
+	if (FCurveKeyData* Change_Add = InAppendedChanges.AddKeysData.SavedCurveState.Find(InCurveId))
+	{
+		InOutTarget.AddKeysData.SavedCurveState.Add(InCurveId, MoveTemp(*Change_Add));
+	}
+	if (FCurveKeyData* Change_Remove = InAppendedChanges.RemoveKeysData.SavedCurveState.Find(InCurveId))
+	{
+		InOutTarget.RemoveKeysData.SavedCurveState.Add(InCurveId, MoveTemp(*Change_Remove));
+	}
+	if (FKeyAttributeChangeData_PerCurve* Change_Move = InAppendedChanges.KeyAttributeData.ChangedCurves.Find(InCurveId))
+	{
+		InOutTarget.KeyAttributeData.ChangedCurves.Add(InCurveId, MoveTemp(*Change_Move));
+	}
+	if (FCurveAttributeChangeData_PerCurve* Change_Move = InAppendedChanges.CurveAttributeData.ChangeData.Find(InCurveId))
+	{
+		InOutTarget.CurveAttributeData.ChangeData.Add(InCurveId, MoveTemp(*Change_Move));
+	}
+}
+}
+	
+void FMultiGenericCommandAppender::ProcessChange(const FCurveModelID& InCurveModel, TFunctionRef<void(ISingleCurveChangeVisitor&)> InProcessCallback)
+{
+	FGenericCurveChangeData CurveChanges;
+	FAppendPerCurveGenericChangeVisitor Visitor(InCurveModel, CurveChanges);
+	InProcessCallback(Visitor);
+	
+	if (CurveChanges)
+	{
+		ProcessDiff(InCurveModel, MoveTemp(CurveChanges));
+	}
+}
+
+void FMultiGenericCommandAppender::PostProcessChanges()
+{
+	// If there were any FCurveModels that are not bound to any UObject, just append to the dummy transaction object.
+	// There is no package so we don't need to worry about any dirty flags.
+	if (ChangesToCurvesWithoutUObject)
+	{
+		CurveEditor->GetTransactionManager()->AppendChange(
+			MakeUnique<FGenericCurveChangeCommand>(CurveEditor, MoveTemp(ChangesToCurvesWithoutUObject))
+			);
+	}
+}
+
+void FMultiGenericCommandAppender::ProcessDiff(const FCurveModelID& InCurveId, FGenericCurveChangeData&& InData)
+{
+	check(InData.HasChanges());
+	
+	FCurveModel* CurveModel = CurveEditor->FindCurve(InCurveId);
+	if (!CurveModel)
+	{
+		return;
+	}
+
+	if (UObject* OwningObject = CurveModel->GetOwningObject())
+	{
+		CurveEditor->GetTransactionManager()->AppendChange(
+			OwningObject, MakeUnique<FGenericCurveChangeCommand>(CurveEditor, MoveTemp(InData))
+			);
+	}
+	else
+	{
+		Private::AppendChanges(InCurveId, ChangesToCurvesWithoutUObject, MoveTemp(InData));
+	}
+	
+	AllocationSize += ChangesToCurvesWithoutUObject.GetAllocatedSize();
+}
+}
diff --git a/Engine/Source/Editor/CurveEditor/Private/Modification/Keys/Diff/MultiGenericCommandAppender.h b/Engine/Source/Editor/CurveEditor/Private/Modification/Keys/Diff/MultiGenericCommandAppender.h
new file mode 100644
--- /dev/null
+++ b/Engine/Source/Editor/CurveEditor/Private/Modification/Keys/Diff/MultiGenericCommandAppender.h
@@ -0,0 +1,49 @@
+﻿// Copyright Epic Games, Inc. All Rights Reserved.
+
+#pragma once
+
+#include "Modification/Keys/Data/GenericCurveChangeData.h"
+#include "Modification/Keys/Diff/IMultiCurveChangeVisitor.h"
+
+class FCurveEditor;
+
+namespace UE::CurveEditor
+{
+class FTransactionManager;
+struct FGenericCurveChangeData;
+
+/**
+ * Handles appending FGenericCurveChangeCommands to the transaction system.
+ * - StoreUndo is called for each FCurveModel that returns non-null from FCurveModel::GetOwningObject.
+ * These objects will have their package's dirty flag stored to the undo buffer.
+ * - One catch all FGenericCurveChangeCommand is created for all other FCurveModels that returned null from FCurveModel::GetOwningObject.
+ */
+class FMultiGenericCommandAppender : public IMultiCurveChangeVisitor
+{
+public:
+
+	explicit FMultiGenericCommandAppender(TSharedPtr<FCurveEditor> InCurveEditor)
+		: CurveEditor(MoveTemp(InCurveEditor))
+	{}
+
+	SIZE_T GetAllocationSize() const { return AllocationSize; }
+
+	//~ Begin ICurvesDiffer Interface
+	virtual void ProcessChange(const FCurveModelID& InCurveModel, TFunctionRef<void(ISingleCurveChangeVisitor&)> InProcessCallback) override;
+	virtual void PostProcessChanges() override;
+	//~ End ICurvesDiffer Interface
+
+private:
+
+	/** The curve editor to which we append the change commands. */
+	TSharedPtr<FCurveEditor> CurveEditor;
+
+	/** Holds the changes made to FCurveModel's that returned no UObject from FCurveModel::GetOwningObject. */
+	FGenericCurveChangeData ChangesToCurvesWithoutUObject;
+
+	/** How much memory was allocated by all undo changes. */
+	SIZE_T AllocationSize = 0;
+
+	void ProcessDiff(const FCurveModelID& InCurveId, FGenericCurveChangeData&& InData);
+};
+}
diff --git a/Engine/Source/Editor/CurveEditor/Private/Modification/Keys/GenericCurveChangeCommand.cpp b/Engine/Source/Editor/CurveEditor/Private/Modification/Keys/GenericCurveChangeCommand.cpp
--- a/Engine/Source/Editor/CurveEditor/Private/Modification/Keys/GenericCurveChangeCommand.cpp
+++ b/Engine/Source/Editor/CurveEditor/Private/Modification/Keys/GenericCurveChangeCommand.cpp
@@ -9,7 +9,7 @@ namespace UE::CurveEditor
 {
 void FGenericCurveChangeCommand::Apply(UObject* Object)
 {
-	if (const TSharedPtr<FCurveEditor> CurveEditor = GetCurveEditor(Object))
+	if (const TSharedPtr<FCurveEditor> CurveEditor = GetCurveEditor())
 	{
 		GenericCurveChange::ApplyChange(CurveEditor->GetCurves(), DeltaChange);
 	}
@@ -17,7 +17,7 @@ void FGenericCurveChangeCommand::Apply(UObject* Object)
 
 void FGenericCurveChangeCommand::Revert(UObject* Object)
 {
-	if (const TSharedPtr<FCurveEditor> CurveEditor = GetCurveEditor(Object))
+	if (const TSharedPtr<FCurveEditor> CurveEditor = GetCurveEditor())
 	{
 		GenericCurveChange::RevertChange(CurveEditor->GetCurves(), DeltaChange);
 	}
diff --git a/Engine/Source/Editor/CurveEditor/Private/Modification/Keys/GenericCurveChangeCommand.h b/Engine/Source/Editor/CurveEditor/Private/Modification/Keys/GenericCurveChangeCommand.h
--- a/Engine/Source/Editor/CurveEditor/Private/Modification/Keys/GenericCurveChangeCommand.h
+++ b/Engine/Source/Editor/CurveEditor/Private/Modification/Keys/GenericCurveChangeCommand.h
@@ -19,7 +19,10 @@ class FGenericCurveChangeCommand : public FCurveEditorCommandChange
 {
 public:
 
-	explicit FGenericCurveChangeCommand(FGenericCurveChangeData InDeltaChange) : DeltaChange(MoveTemp(InDeltaChange)) {}
+	explicit FGenericCurveChangeCommand(TSharedPtr<FCurveEditor> InCurveEditor, FGenericCurveChangeData InDeltaChange)
+		: FCurveEditorCommandChange(MoveTemp(InCurveEditor))
+		, DeltaChange(MoveTemp(InDeltaChange))
+	{}
 
 	//~ Begin FChange Interface
 	virtual void Apply(UObject* Object) override;
diff --git a/Engine/Source/Editor/CurveEditor/Private/Modification/Selection/SelectionChangeCommand.cpp b/Engine/Source/Editor/CurveEditor/Private/Modification/Selection/SelectionChangeCommand.cpp
--- a/Engine/Source/Editor/CurveEditor/Private/Modification/Selection/SelectionChangeCommand.cpp
+++ b/Engine/Source/Editor/CurveEditor/Private/Modification/Selection/SelectionChangeCommand.cpp
@@ -9,7 +9,7 @@ namespace UE::CurveEditor
 {
 void FSelectionChangeCommand::Apply(UObject* Object)
 {
-	if (const TSharedPtr<FCurveEditor> CurveEditor = GetCurveEditor(Object))
+	if (const TSharedPtr<FCurveEditor> CurveEditor = GetCurveEditor())
 	{
 		ApplySelectionChange(CurveEditor->Selection, DeltaChange);
 	}
@@ -17,7 +17,7 @@ void FSelectionChangeCommand::Apply(UObject* Object)
 
 void FSelectionChangeCommand::Revert(UObject* Object)
 {
-	if (const TSharedPtr<FCurveEditor> CurveEditor = GetCurveEditor(Object))
+	if (const TSharedPtr<FCurveEditor> CurveEditor = GetCurveEditor())
 	{
 		RevertSelectionChange(CurveEditor->Selection, DeltaChange);
 	}
diff --git a/Engine/Source/Editor/CurveEditor/Private/Modification/Selection/SelectionChangeCommand.h b/Engine/Source/Editor/CurveEditor/Private/Modification/Selection/SelectionChangeCommand.h
--- a/Engine/Source/Editor/CurveEditor/Private/Modification/Selection/SelectionChangeCommand.h
+++ b/Engine/Source/Editor/CurveEditor/Private/Modification/Selection/SelectionChangeCommand.h
@@ -12,7 +12,10 @@ class FSelectionChangeCommand : public FCurveEditorCommandChange
 {
 public:
 
-	explicit FSelectionChangeCommand(FSelectionDeltaChange InDeltaChange) : DeltaChange(MoveTemp(InDeltaChange)) {}
+	explicit FSelectionChangeCommand(TSharedPtr<FCurveEditor> InCurveEditor, FSelectionDeltaChange InDeltaChange)
+		: FCurveEditorCommandChange(MoveTemp(InCurveEditor))
+		, DeltaChange(MoveTemp(InDeltaChange))
+	{}
 	
 	//~ Begin FChange Interface
 	virtual void Apply(UObject* Object) override;
diff --git a/Engine/Source/Editor/CurveEditor/Private/Modification/TransactionManager.cpp b/Engine/Source/Editor/CurveEditor/Private/Modification/TransactionManager.cpp
--- a/Engine/Source/Editor/CurveEditor/Private/Modification/TransactionManager.cpp
+++ b/Engine/Source/Editor/CurveEditor/Private/Modification/TransactionManager.cpp
@@ -15,17 +15,22 @@ FTransactionManager::FTransactionManager(TWeakPtr<FCurveEditor> InCurveEditor)
 {
 	DummyTransactionObject->OwningCurveEditor = MoveTemp(InCurveEditor);
 }
-
-void FTransactionManager::AppendChange(TUniquePtr<FCurveEditorCommandChange> InChange, const FText& InDescription) const
+	
+void FTransactionManager::AppendChange(UObject* InObject, TUniquePtr<FCurveEditorCommandChange> InChange, const FText& InDescription) const
 {
 	FScopedTransaction Transaction(InDescription);
 	if (GUndo // GUndo can be reset by external systems, e.g. during world tear down.
-		&& ensure(DummyTransactionObject) && ensure(InChange))
+		&& ensure(InChange))
 	{
-		GUndo->StoreUndo(DummyTransactionObject, MoveTemp(InChange));
+		GUndo->StoreUndo(InObject, MoveTemp(InChange));
 		OnCommandAppendedDelegate.Broadcast();
 	}
 }
+	
+void FTransactionManager::AppendChange(TUniquePtr<FCurveEditorCommandChange> InChange, const FText& InDescription) const
+{
+	AppendChange(DummyTransactionObject, MoveTemp(InChange), InDescription);
+}
 
 void FTransactionManager::AddReferencedObjects(FReferenceCollector& Collector)
 {
diff --git a/Engine/Source/Editor/CurveEditor/Private/Modification/Utils/ScopedCurveChange.cpp b/Engine/Source/Editor/CurveEditor/Private/Modification/Utils/ScopedCurveChange.cpp
--- a/Engine/Source/Editor/CurveEditor/Private/Modification/Utils/ScopedCurveChange.cpp
+++ b/Engine/Source/Editor/CurveEditor/Private/Modification/Utils/ScopedCurveChange.cpp
@@ -14,6 +14,8 @@
 #include "Modification/Keys/SnapshotPerfDebugUtils.h"
 #include "ProfilingDebugging/ScopedTimers.h"
 #include "SCurveEditor.h"
+#include "Modification/Keys/Diff/MultiGenericCommandAppender.h"
+#include "UObject/Object.h"
 
 namespace UE::CurveEditor
 {
@@ -90,30 +92,27 @@ FScopedCurveChange::~FScopedCurveChange()
 	}
 
 	FAutoScopedDurationTimer TimeMeasurement;
-	FGenericCurveChangeData DeltaChange = ChangeDiff.ComputeDiff();
-	const auto LogPerfData = [this, &TimeMeasurement, &DeltaChange]
+	const auto LogPerfData = [this, &TimeMeasurement](SIZE_T InAllocationSize)
 	{
 		UE_CLOG(CVarLogUndoKeyChangeSize.GetValueOnAnyThread(), LogCurveEditor, Log, TEXT("Change'%s' --- Undo size: %s ---\t Diff: Time %f ms, %s"),
 			GUndo ? *GUndo->GetContext().Title.ToString() : *Description.ToString(),
-			*FText::AsMemory(DeltaChange.GetAllocatedSize()).ToString(),
+			*FText::AsMemory(InAllocationSize).ToString(),
 			FTimespan::FromSeconds(TimeMeasurement.GetTime()).GetTotalMilliseconds(),
 			*DumpSnapshotPerfData(ChangeDiff.GetSnapshot(), ECurvesSnapshotPerfFlags::All)
 			); 
 	};
 	
-	if (!DeltaChange)
-	{
-		return;
-	}
-	
-	LogPerfData(); // Important to do this before MoveTemp
 	if (!bHasBeenCancelled)
 	{
-		TryAppendCommand(MakeUnique<FGenericCurveChangeCommand>(MoveTemp(DeltaChange)));
+		FMultiGenericCommandAppender CommandAppender(CurveEditorPin);
+		ChangeDiff.ProcessDiffs(CommandAppender);
+		LogPerfData(CommandAppender.GetAllocationSize()); 
 	}
 	else
 	{
+		FGenericCurveChangeData DeltaChange = ChangeDiff.ComputeDiff();
 		GenericCurveChange::RevertChange(*CurveEditorPin, DeltaChange);
+		LogPerfData(DeltaChange.GetAllocatedSize()); 
 	}
 }
 }
diff --git a/Engine/Source/Editor/CurveEditor/Private/Modification/Utils/ScopedSelectionChange.cpp b/Engine/Source/Editor/CurveEditor/Private/Modification/Utils/ScopedSelectionChange.cpp
--- a/Engine/Source/Editor/CurveEditor/Private/Modification/Utils/ScopedSelectionChange.cpp
+++ b/Engine/Source/Editor/CurveEditor/Private/Modification/Utils/ScopedSelectionChange.cpp
@@ -44,7 +44,7 @@ FScopedSelectionChange::~FScopedSelectionChange()
 	FSelectionDeltaChange DeltaChange = DiffSelection(OriginalSelection, CurveEditorPin->Selection);
 	if (!bHasBeenCancelled)
 	{
-		TryAppendCommand(MakeUnique<FSelectionChangeCommand>(MoveTemp(DeltaChange)));
+		TryAppendCommand(MakeUnique<FSelectionChangeCommand>(CurveEditorPin, MoveTemp(DeltaChange)));
 	}
 	else
 	{
diff --git a/Engine/Source/Editor/CurveEditor/Public/Modification/CurveEditorCommandChange.h b/Engine/Source/Editor/CurveEditor/Public/Modification/CurveEditorCommandChange.h
--- a/Engine/Source/Editor/CurveEditor/Public/Modification/CurveEditorCommandChange.h
+++ b/Engine/Source/Editor/CurveEditor/Public/Modification/CurveEditorCommandChange.h
@@ -3,7 +3,7 @@
 #pragma once
 
 #include "Misc/Change.h"
-#include "Templates/SharedPointerFwd.h"
+#include "Templates/SharedPointer.h"
 
 #define UE_API CURVEEDITOR_API
 
@@ -16,6 +16,10 @@ class FCurveEditorCommandChange : public FCommandChange
 {
 public:
 
+	explicit FCurveEditorCommandChange(TSharedPtr<FCurveEditor> InCurveEditor)
+		: WeakCurveEditor(MoveTemp(InCurveEditor))
+	{}
+
 	//~ Begin FChange Interface
 	UE_API virtual bool HasExpired(UObject* InObject) const override;
 	UE_API virtual FString ToString() const override;
@@ -24,7 +28,12 @@ public:
 protected:
 
 	/** @return The curve editor associated with Object. */
-	static UE_API TSharedPtr<FCurveEditor> GetCurveEditor(UObject* InObject);
+	TSharedPtr<FCurveEditor> GetCurveEditor() const;
+
+private:
+
+	/** The curve editor this operation was performed on. */
+	TWeakPtr<FCurveEditor> WeakCurveEditor;
 };
 }
 
diff --git a/Engine/Source/Editor/CurveEditor/Public/Modification/Keys/CurveChangeDiff.h b/Engine/Source/Editor/CurveEditor/Public/Modification/Keys/CurveChangeDiff.h
--- a/Engine/Source/Editor/CurveEditor/Public/Modification/Keys/CurveChangeDiff.h
+++ b/Engine/Source/Editor/CurveEditor/Public/Modification/Keys/CurveChangeDiff.h
@@ -4,7 +4,6 @@
 
 #include "Containers/Map.h"
 #include "CurvesSnapshotBuilder.h"
-#include "Data/CurveKeyDataMap.h"
 #include "Data/GenericCurveChangeData.h"
 #include "ECurveChangeFlags.h"
 #include "Templates/SharedPointer.h"
@@ -13,13 +12,20 @@ class FCurveEditor;
 
 namespace UE::CurveEditor
 {
+class IMultiCurveChangeVisitor;
 struct FGenericCurveChangeData;
 
-/** Computes what has changed for the data saved in InSnapshot. */
+/** @return FGenericCurveChangeData containing what has changed for the data saved in InSnapshot. */
 CURVEEDITOR_API FGenericCurveChangeData ComputeDiff(
 	const FCurveEditor& InCurveEditor, const ECurveChangeFlags InOperations, const FCurvesSnapshot& InSnapshot
 	);
 
+/** Computes what has changed for the data saved in InSnapshot and passes those changes to IMultiCurveChangeVisitor to perform any processing. */
+CURVEEDITOR_API	void ProcessChanges(
+	const FCurveEditor& InCurveEditor, const ECurveChangeFlags InOperations, const FCurvesSnapshot& InSnapshot,
+	IMultiCurveChangeVisitor& InChangeVisitor
+	);
+
 /**
  * Utility that holds a snapshot of some curve data for the purposes of later comparing changes made to those curves.
  * 
@@ -47,11 +53,10 @@ public:
 	{}
 
 	/** Computes what has changed on the given curves. */
-	FGenericCurveChangeData ComputeDiff() const
-	{
-		const TSharedPtr<FCurveEditor> CurveEditorPin = WeakCurveEditor.Pin();
-		return CurveEditorPin ? CurveEditor::ComputeDiff(*CurveEditorPin, Operations, Snapshot) : FGenericCurveChangeData();
-	}
+	CURVEEDITOR_API FGenericCurveChangeData ComputeDiff() const;
+
+	/** Passes all changes to the given instance. */
+	CURVEEDITOR_API void ProcessDiffs(IMultiCurveChangeVisitor& InChangeBuilder) const;
 
 	/** The snapshot against which is diffed. */
 	const FCurvesSnapshot& GetSnapshot() const { return Snapshot; }
diff --git a/Engine/Source/Editor/CurveEditor/Public/Modification/Keys/Diff/IMultiCurveChangeVisitor.h b/Engine/Source/Editor/CurveEditor/Public/Modification/Keys/Diff/IMultiCurveChangeVisitor.h
new file mode 100644
--- /dev/null
+++ b/Engine/Source/Editor/CurveEditor/Public/Modification/Keys/Diff/IMultiCurveChangeVisitor.h
@@ -0,0 +1,45 @@
+﻿// Copyright Epic Games, Inc. All Rights Reserved.
+
+#pragma once
+
+#include "Templates/Function.h"
+
+struct FCurveModelID;
+
+namespace UE::CurveEditor
+{
+class ISingleCurveChangeVisitor;
+	
+/**
+ * Interface for processing changes made to multiple FCurveModels.
+ * @see ISingleCurveChangeVisitor
+ * @see FCurveChangeDiff
+ */
+class IMultiCurveChangeVisitor
+{
+public:
+	
+	/**
+	 * Called when changes are about to be processed for InCurveModel.
+	 *
+	 * You are supposed to construct some ICurveChangeVisitor and call InProcessCallback with it.
+	 * InProcessCallback will only use your visitor for the duration of the function call.
+	 * This pattern allows construction of your visitor on the stack (avoiding heap allocating it).
+	 * 
+	 * A sample implementation could be:
+	 * void FFooDiffBuilder::ProcessChange(const FCurveModelID& InCurveModel, TFunctionRef<void(ICurveChangeVisitor&)> InProcessCallback)
+	 * {
+	 *		FFooCurveChangeVisitor MyVisitor(InCurveModel); 
+	 *		InProcessCallback(MyVisitor);
+	 * }
+	 */
+	virtual void ProcessChange(const FCurveModelID& InCurveModel, TFunctionRef<void(ISingleCurveChangeVisitor&)> InProcessCallback) = 0;
+
+	/** Called before any ProcessChange calls are made. */
+	virtual void PreProcessChanges() {}
+	/** Called after all changes have been listed, i.e. no more ProcessChange will follow. */
+	virtual void PostProcessChanges() {}
+
+	virtual ~IMultiCurveChangeVisitor() = default;
+};
+}
diff --git a/Engine/Source/Editor/CurveEditor/Public/Modification/Keys/Diff/ISingleCurveChangeVisitor.h b/Engine/Source/Editor/CurveEditor/Public/Modification/Keys/Diff/ISingleCurveChangeVisitor.h
new file mode 100644
--- /dev/null
+++ b/Engine/Source/Editor/CurveEditor/Public/Modification/Keys/Diff/ISingleCurveChangeVisitor.h
@@ -0,0 +1,38 @@
+﻿// Copyright Epic Games, Inc. All Rights Reserved.
+
+#pragma once
+
+namespace UE::CurveEditor
+{
+struct FCurveAttributeChangeData_PerCurve;
+struct FCurveKeyData;
+struct FKeyAttributeChangeData_PerCurve;
+struct FMoveKeysChangeData_PerCurve;
+
+/**
+ * Interface for processing changes made to a single FCurveModel.
+ * @see ICurveDiffBuilder
+ * @see FCurveChangeDiff
+ */
+class ISingleCurveChangeVisitor
+{
+public:
+
+	/** Processes moved keys. Only called if InData contains any changes, i.e. InData.HasChanges() == true. */
+	virtual void ProcessMoveKeys(FMoveKeysChangeData_PerCurve&& InData) {}
+
+	/** Processes added keys. Only called if InData contains any changes, i.e. InData.HasChanges() == true. */
+	virtual void ProcessAddKeys(FCurveKeyData&& InData) {}
+	
+	/** Processes removed keys. Only called if InData contains any changes, i.e. InData.HasChanges() == true. */
+	virtual void ProcessRemoveKeys(FCurveKeyData&& InData) {}
+	
+	/** Processes changed key attributes. Only called if InData contains any changes, i.e. InData.HasChanges() == true. */
+	virtual void ProcessKeyAttributesChange(FKeyAttributeChangeData_PerCurve&& InData) {}
+
+	/** Processes changed curve attributes. Only called if InData contains any changes, i.e. InData.HasChanges() == true. */
+	virtual void ProcessCurveAttributesChange(FCurveAttributeChangeData_PerCurve&& InData) {}
+	
+	virtual ~ISingleCurveChangeVisitor() = default;
+};
+}
\ No newline at end of file
diff --git a/Engine/Source/Editor/CurveEditor/Public/Modification/TransactionManager.h b/Engine/Source/Editor/CurveEditor/Public/Modification/TransactionManager.h
--- a/Engine/Source/Editor/CurveEditor/Public/Modification/TransactionManager.h
+++ b/Engine/Source/Editor/CurveEditor/Public/Modification/TransactionManager.h
@@ -3,6 +3,7 @@
 #pragma once
 
 #include "CurveEditorCommandChange.h"
+#include "Internationalization/Text.h"
 #include "Templates/UnrealTemplate.h"
 #include "UObject/GCObject.h"
 #include "UObject/ObjectPtr.h"
@@ -21,8 +22,22 @@ public:
 
 	UE_API explicit FTransactionManager(TWeakPtr<FCurveEditor> InCurveEditor);
 	
+
+	/**
+	 * Appends a change to the current transaction and associates it with an UObject.
+	 * 
+	 * You should prefer this overload whenever possible because it causes the object's package's dirty flag to be saved and restore by the undo system.
+	 * Example: If you're modifying a FCurveModel, you'd pass in FCurveModel::GetOwningObject.
+	 *
+	 * @param InObject The object being modified
+	 * @param InChange The change to associate with the currently open transaction
+	 * @param InDescription Description of the change
+	 */
+	UE_API void AppendChange(UObject* InObject, TUniquePtr<FCurveEditorCommandChange> InChange, const FText& InDescription = FText::GetEmpty()) const;
+	
 	/**
 	 * Appends a change to the current transaction.
+	 * Use this version when there is no UObject the change can be associated with - otherwise use the UObject overload.
 	 *
 	 * The transaction can consist of multiple changes, e.g.
 	 * - system 1 may move keys, and
@@ -31,7 +46,7 @@ public:
 	 * @param InChange The change to associate with the currently open transaction
 	 * @param InDescription Description of the change
 	 */
-	UE_API void AppendChange(TUniquePtr<FCurveEditorCommandChange> InChange, const FText& InDescription) const;
+	UE_API void AppendChange(TUniquePtr<FCurveEditorCommandChange> InChange, const FText& InDescription = FText::GetEmpty()) const;
 
 	DECLARE_MULTICAST_DELEGATE(FOnCommandAppended);
 	/** Invoked after a command has been appended to the undo stack. */
