Curve Editor: Baking keys with a value different than 1 would increase range
#trivial

diff --git a/Engine/Source/Editor/CurveEditor/Private/Filters/CurveEditorBakeFilter.cpp b/Engine/Source/Editor/CurveEditor/Private/Filters/CurveEditorBakeFilter.cpp
--- a/Engine/Source/Editor/CurveEditor/Private/Filters/CurveEditorBakeFilter.cpp
+++ b/Engine/Source/Editor/CurveEditor/Private/Filters/CurveEditorBakeFilter.cpp
@@ -178,9 +178,7 @@ void UCurveEditorBakeFilter::ApplyFilter_Impl(TSharedRef<FCurveEditor> InCurveEd
 			++KeyIndex;
 		}
 
-		int32 NumKeysToAdd = bUseSeconds ? FMath::RoundToInt((SampleMaxKey - SampleMinKey) / Interval) + 1 :
-			(MaxInDisplayFrames.FrameNumber.Value - MinInDisplayFrames.FrameNumber.Value)  + 1;
-
+		int32 NumKeysToAdd = FMath::RoundToInt((SampleMaxKey - SampleMinKey) / Interval) + 1;
 	
 		NewKeyPositions.Reset(NumKeysToAdd);
 		NewKeyAttributes.Reset(NumKeysToAdd);
@@ -194,12 +192,18 @@ void UCurveEditorBakeFilter::ApplyFilter_Impl(TSharedRef<FCurveEditor> InCurveEd
 		for (; KeyIndex < NumKeysToAdd; ++KeyIndex)
 		{
 			double NewTime = SampleMinKey + KeyIndex * Interval;
-
-			FKeyPosition CurrentKey(NewTime, 0.0);
-			if (Curve->Evaluate(CurrentKey.InputValue, CurrentKey.OutputValue))
+			if (NewTime <= SampleMaxKey)
+			{
+				FKeyPosition CurrentKey(NewTime, 0.0);
+				if (Curve->Evaluate(CurrentKey.InputValue, CurrentKey.OutputValue))
+				{
+					NewKeyPositions.Add(CurrentKey);
+					NewKeyAttributes.Add(DefaultKeyAttribute);
+				}
+			}
+			else
 			{
-				NewKeyPositions.Add(CurrentKey);
-				NewKeyAttributes.Add(DefaultKeyAttribute);
+				break;
 			}
 		}
 
