Curve Editor: When selection is empty, you can now use the combo button to change the tangent mode of all keys in view.

diff --git a/Engine/Source/Editor/CurveEditor/Private/Misc/AttributeAccumulationUtils.cpp b/Engine/Source/Editor/CurveEditor/Private/Misc/AttributeAccumulationUtils.cpp
new file mode 100644
--- /dev/null
+++ b/Engine/Source/Editor/CurveEditor/Private/Misc/AttributeAccumulationUtils.cpp
@@ -0,0 +1,96 @@
+// Copyright Epic Games, Inc. All Rights Reserved.
+
+#include "AttributeAccumulationUtils.h"
+
+#include "CurveEditor.h"
+#include "CurveEditorSelectionPrivate.h"
+#include "CurveModel.h"
+
+namespace UE::CurveEditor
+{
+void UpdateCommonCurveInfo(
+	const FCurveEditor& InCurveEditor, bool& bOutSelectionSupportsWeightedTangents,
+	FCurveAttributes& OutCachedCommonCurveAttributes, FKeyAttributes& OutCachedCommonKeyAttributes)
+{	
+	FAccumulatedCurveInfo AccumulatedCurveInfo = InCurveEditor.Selection.IsEmpty() ?
+		AccumulateAllCurvesInfo(InCurveEditor) :
+		AccumulateSelectedCurvesInfo(InCurveEditor);
+
+	// Reset the common curve and key info
+	bOutSelectionSupportsWeightedTangents = AccumulatedCurveInfo.bSelectionSupportsWeightedTangents;
+	OutCachedCommonCurveAttributes = AccumulatedCurveInfo.CurveAttributes.Get(FCurveAttributes());
+	OutCachedCommonKeyAttributes = AccumulatedCurveInfo.KeyAttributes.Get(FKeyAttributes());
+	
+}
+
+FAccumulatedCurveInfo AccumulateSelectedCurvesInfo(const FCurveEditor& InCurveEditor)
+{
+	FAccumulatedCurveInfo AccumulatedAttributes;
+	for (const TTuple<FCurveModelID, FKeyHandleSet>& Pair : InCurveEditor.Selection.GetAll())
+	{
+		if (const FCurveModel* Curve = InCurveEditor.FindCurve(Pair.Key))
+		{
+			AccumulateCurveInfo(Curve, Pair.Value.AsArray(), AccumulatedAttributes);
+		}
+	}
+	return AccumulatedAttributes;
+}
+
+FAccumulatedCurveInfo AccumulateAllCurvesInfo(const FCurveEditor& InCurveEditor)
+{	
+	FAccumulatedCurveInfo AccumulatedAttributes;
+	for (const FCurveModelID ID : InCurveEditor.GetEditedCurves())
+	{
+		if (FCurveModel* Curve = InCurveEditor.FindCurve(ID))
+		{
+			const TArray<FKeyHandle> KeyHandles = Curve->GetAllKeys();
+			AccumulateCurveInfo(Curve, KeyHandles, AccumulatedAttributes);
+		}
+	}	
+	return AccumulatedAttributes;
+}
+
+void AccumulateCurveInfo(
+	const FCurveModel* InCurve, TConstArrayView<const FKeyHandle> InKeys, FAccumulatedCurveInfo& InOutAccumulatedAttributes)
+{
+	FCurveAttributes CurveAttributes;
+	InCurve->GetCurveAttributes(CurveAttributes);
+
+	// Some curves don't support extrapolation. We don't count them for determining the accumulated state.
+	const bool bCurveHasNoExtrapolation = CurveAttributes.HasPreExtrapolation() && CurveAttributes.GetPreExtrapolation() == RCCE_None
+		&& CurveAttributes.HasPostExtrapolation() && CurveAttributes.GetPostExtrapolation() == RCCE_None;
+	if (!bCurveHasNoExtrapolation && !InOutAccumulatedAttributes.CurveAttributes.IsSet())
+	{
+		InOutAccumulatedAttributes.CurveAttributes = CurveAttributes;
+	}
+	else if (!bCurveHasNoExtrapolation)
+	{
+		InOutAccumulatedAttributes.CurveAttributes = FCurveAttributes::MaskCommon(
+			InOutAccumulatedAttributes.CurveAttributes.GetValue(), CurveAttributes
+			);
+	}
+		
+	TArray<FKeyAttributes> AllKeyAttributes;	
+	AllKeyAttributes.SetNum(InKeys.Num());
+
+	InCurve->GetKeyAttributes(InKeys, AllKeyAttributes);
+	for (const FKeyAttributes& Attributes : AllKeyAttributes)
+	{
+		if (Attributes.HasTangentWeightMode())
+		{
+			InOutAccumulatedAttributes.bSelectionSupportsWeightedTangents = true;
+		}
+
+		if (!InOutAccumulatedAttributes.KeyAttributes.IsSet())
+		{
+			InOutAccumulatedAttributes.KeyAttributes = Attributes;
+		}
+		else
+		{
+			InOutAccumulatedAttributes.KeyAttributes = FKeyAttributes::MaskCommon(
+				InOutAccumulatedAttributes.KeyAttributes.GetValue(), Attributes
+				);
+		}
+	}
+}
+}
\ No newline at end of file
diff --git a/Engine/Source/Editor/CurveEditor/Private/Misc/AttributeAccumulationUtils.h b/Engine/Source/Editor/CurveEditor/Private/Misc/AttributeAccumulationUtils.h
new file mode 100644
--- /dev/null
+++ b/Engine/Source/Editor/CurveEditor/Private/Misc/AttributeAccumulationUtils.h
@@ -0,0 +1,32 @@
+// Copyright Epic Games, Inc. All Rights Reserved.
+
+#pragma once
+
+#include "Curves/RealCurve.h"
+#include "CurveDataAbstraction.h"
+
+class FCurveEditor;
+class FCurveModel;
+
+namespace  UE::CurveEditor
+{
+	/** Used to hold curve and key attributes to update the Common Curve Info with. */
+	struct FAccumulatedCurveInfo
+	{
+		bool bSelectionSupportsWeightedTangents = false;
+		TOptional<FCurveAttributes> CurveAttributes;
+		TOptional<FKeyAttributes> KeyAttributes;
+	};
+
+	/** Update current curve selection attributes, if no keys are selected, info of all curves currently contained in the Curve Editor is used. */
+	void UpdateCommonCurveInfo(const FCurveEditor& InCurveEditor, bool& bOutSelectionSupportsWeightedTangents, FCurveAttributes& OutCachedCommonCurveAttributes, FKeyAttributes& OutCachedCommonKeyAttributes);
+
+	/** Gather curve and key attributes from all currently selected curves. */
+	FAccumulatedCurveInfo AccumulateSelectedCurvesInfo(const FCurveEditor& InCurveEditor);
+	/** Gather curve and key attributes from all curves currently contained in the Curve Editor regardless of visibility. .*/
+	FAccumulatedCurveInfo AccumulateAllCurvesInfo(const FCurveEditor& InCurveEditor);
+
+	/** Add curve and key attributes from the given curve and the keys on it to the accumulated info. */
+	void AccumulateCurveInfo(const FCurveModel* InCurve, TConstArrayView<const FKeyHandle> InKeys, FAccumulatedCurveInfo& InOutAccumulatedAttributes);
+}
+
diff --git a/Engine/Source/Editor/CurveEditor/Private/SCurveEditorPanel.cpp b/Engine/Source/Editor/CurveEditor/Private/SCurveEditorPanel.cpp
--- a/Engine/Source/Editor/CurveEditor/Private/SCurveEditorPanel.cpp
+++ b/Engine/Source/Editor/CurveEditor/Private/SCurveEditorPanel.cpp
@@ -64,8 +64,10 @@
 #include "SGridLineSpacingList.h"
 #include "ScopedTransaction.h"
 #include "SlotBase.h"
+#include "Algo/AnyOf.h"
 #include "Filters/CurveEditorEulerFilter.h"
 #include "Filters/PromotedFilterContainer.h"
+#include "Misc/AttributeAccumulationUtils.h"
 #include "Misc/ResizeParamUtils.h"
 #include "Modification/Utils/ScopedCurveChange.h"
 #include "Modification/Utils/ScopedSelectionChange.h"
@@ -347,12 +349,12 @@ void SCurveEditorPanel::BindCommands()
 {
 	// Interpolation and tangents
 	{
-		FExecuteAction SetConstant   = FExecuteAction::CreateSP(this, &SCurveEditorPanel::SetKeyAttributes, FKeyAttributes().SetInterpMode(RCIM_Constant).SetTangentMode(RCTM_Auto), LOCTEXT("SetInterpConstant", "Set Interp Constant"));
-		FExecuteAction SetLinear     = FExecuteAction::CreateSP(this, &SCurveEditorPanel::SetKeyAttributes, FKeyAttributes().SetInterpMode(RCIM_Linear).SetTangentMode(RCTM_Auto),   LOCTEXT("SetInterpLinear",   "Set Interp Linear"));
-		FExecuteAction SetCubicAuto  = FExecuteAction::CreateSP(this, &SCurveEditorPanel::SetKeyAttributes, FKeyAttributes().SetInterpMode(RCIM_Cubic).SetTangentMode(RCTM_Auto),    LOCTEXT("SetInterpCubic",    "Set Interp Auto"));
-		FExecuteAction SetCubicSmartAuto = FExecuteAction::CreateSP(this, &SCurveEditorPanel::SetKeyAttributes, FKeyAttributes().SetInterpMode(RCIM_Cubic).SetTangentMode(RCTM_SmartAuto),	LOCTEXT("SetInterpSmartAuto", "Set Interp Smart Auto"));
-		FExecuteAction SetCubicUser  = FExecuteAction::CreateSP(this, &SCurveEditorPanel::SetKeyAttributes, FKeyAttributes().SetInterpMode(RCIM_Cubic).SetTangentMode(RCTM_User),    LOCTEXT("SetInterpUser",     "Set Interp User"));
-		FExecuteAction SetCubicBreak = FExecuteAction::CreateSP(this, &SCurveEditorPanel::SetKeyAttributes, FKeyAttributes().SetInterpMode(RCIM_Cubic).SetTangentMode(RCTM_Break),   LOCTEXT("SetInterpBreak",    "Set Interp Break"));
+		FExecuteAction SetConstant   = FExecuteAction::CreateSP(this, &SCurveEditorPanel::HandleSetKeyAttributesCommand, FKeyAttributes().SetInterpMode(RCIM_Constant).SetTangentMode(RCTM_Auto), LOCTEXT("SetInterpConstant", "Set Interp Constant"));
+		FExecuteAction SetLinear     = FExecuteAction::CreateSP(this, &SCurveEditorPanel::HandleSetKeyAttributesCommand, FKeyAttributes().SetInterpMode(RCIM_Linear).SetTangentMode(RCTM_Auto),   LOCTEXT("SetInterpLinear",   "Set Interp Linear"));
+		FExecuteAction SetCubicAuto  = FExecuteAction::CreateSP(this, &SCurveEditorPanel::HandleSetKeyAttributesCommand, FKeyAttributes().SetInterpMode(RCIM_Cubic).SetTangentMode(RCTM_Auto),    LOCTEXT("SetInterpCubic",    "Set Interp Auto"));
+		FExecuteAction SetCubicSmartAuto = FExecuteAction::CreateSP(this, &SCurveEditorPanel::HandleSetKeyAttributesCommand, FKeyAttributes().SetInterpMode(RCIM_Cubic).SetTangentMode(RCTM_SmartAuto),	LOCTEXT("SetInterpSmartAuto", "Set Interp Smart Auto"));
+		FExecuteAction SetCubicUser  = FExecuteAction::CreateSP(this, &SCurveEditorPanel::HandleSetKeyAttributesCommand, FKeyAttributes().SetInterpMode(RCIM_Cubic).SetTangentMode(RCTM_User),    LOCTEXT("SetInterpUser",     "Set Interp User"));
+		FExecuteAction SetCubicBreak = FExecuteAction::CreateSP(this, &SCurveEditorPanel::HandleSetKeyAttributesCommand, FKeyAttributes().SetInterpMode(RCIM_Cubic).SetTangentMode(RCTM_Break),   LOCTEXT("SetInterpBreak",    "Set Interp Break"));
 
 		FExecuteAction    ToggleWeighted    = FExecuteAction::CreateSP(this, &SCurveEditorPanel::ToggleWeightedTangents);
 		FCanExecuteAction CanToggleWeighted = FCanExecuteAction::CreateSP(this, &SCurveEditorPanel::CanToggleWeightedTangents);
@@ -527,7 +529,7 @@ void SCurveEditorPanel::Tick(const FGeometry& AllottedGeometry, const double InC
 		CachedActiveCurvesSerialNumber = CurveEditor->GetActiveCurvesSerialNumber();
 	}
 
-	UpdateCommonCurveInfo();
+	UE::CurveEditor::UpdateCommonCurveInfo(*CurveEditor, bSelectionSupportsWeightedTangents, CachedCommonCurveAttributes, CachedCommonKeyAttributes);
 	UpdateEditBox();
 	UpdateTime();
 
@@ -746,75 +748,6 @@ void SCurveEditorPanel::RebuildCurveViews()
 	OnPostRebuildCurveViewsDelegate.Broadcast();
 }
 
-void SCurveEditorPanel::UpdateCommonCurveInfo()
-{
-	// Gather up common extended curve info for the current set of curves
-	TOptional<FCurveAttributes> AccumulatedCurveAttributes;
-	for (const TTuple<FCurveModelID, FKeyHandleSet>& Pair : CurveEditor->Selection.GetAll())
-	{
-		FCurveAttributes Attributes;
-		
-		FCurveModel* Curve = CurveEditor->FindCurve(Pair.Key);
-		if (Curve)
-		{
-			Curve->GetCurveAttributes(Attributes);
-
-			// Some curves don't support extrapolation. We don't count them for determine the accumulated state.
-			if (Attributes.HasPreExtrapolation() && Attributes.GetPreExtrapolation() == RCCE_None && Attributes.HasPostExtrapolation() && Attributes.GetPostExtrapolation() == RCCE_None)
-			{
-				continue;
-			}
-
-			if (!AccumulatedCurveAttributes.IsSet())
-			{
-				AccumulatedCurveAttributes = Attributes;
-			}
-			else
-			{
-				AccumulatedCurveAttributes = FCurveAttributes::MaskCommon(AccumulatedCurveAttributes.GetValue(), Attributes);
-			}
-		}
-	}
-
-	// Reset the common curve and key info
-	bSelectionSupportsWeightedTangents = false;
-	CachedCommonCurveAttributes = AccumulatedCurveAttributes.Get(FCurveAttributes());
-
-	TOptional<FKeyAttributes> AccumulatedKeyAttributes;
-	TArray<FKeyAttributes> AllKeyAttributes;
-
-	for (const TTuple<FCurveModelID, FKeyHandleSet>& Pair : CurveEditor->Selection.GetAll())
-	{
-		FCurveModel* Curve = CurveEditor->FindCurve(Pair.Key);
-		if (Curve)
-		{
-			AllKeyAttributes.Reset();
-			AllKeyAttributes.SetNum(Pair.Value.Num());
-
-			Curve->GetKeyAttributes(Pair.Value.AsArray(), AllKeyAttributes);
-			for (const FKeyAttributes& Attributes : AllKeyAttributes)
-			{
-				if (Attributes.HasTangentWeightMode())
-				{
-					bSelectionSupportsWeightedTangents = true;
-				}
-
-				if (!AccumulatedKeyAttributes.IsSet())
-				{
-					AccumulatedKeyAttributes = Attributes;
-				}
-				else
-				{
-					AccumulatedKeyAttributes = FKeyAttributes::MaskCommon(AccumulatedKeyAttributes.GetValue(), Attributes);
-				}
-			}
-		}
-	}
-
-	// Reset the common curve and key info
-	CachedCommonKeyAttributes = AccumulatedKeyAttributes.Get(FKeyAttributes());
-}
-
 void SCurveEditorPanel::OnCurveEditorToolChanged(FCurveEditorToolID InToolId)
 {
 	ToolPropertiesPanel->OnToolChanged(InToolId);
@@ -947,7 +880,6 @@ FText SCurveEditorPanel::GetTangentModeLabel() const
 {
 	switch (DetermineTangentMode())
 	{
-	case ETangentModeComboState::NoSelection: return LOCTEXT("TangentMode.NoSelection.Label", "Mode");
 	case ETangentModeComboState::Constant: return LOCTEXT("TangentMode.Constant.Label", "Constant");
 	case ETangentModeComboState::Linear: return LOCTEXT("TangentMode.Linear.Label", "Linear");
 	case ETangentModeComboState::CubicAuto: return LOCTEXT("TangentMode.CubicAuto.Label", "Cubic (auto)");
@@ -955,6 +887,7 @@ FText SCurveEditorPanel::GetTangentModeLabel() const
 	case ETangentModeComboState::CubicUser: return LOCTEXT("TangentMode.CubicUser.Label", "Cubic (flat)");
 	case ETangentModeComboState::CubicBreak: return LOCTEXT("TangentMode.CubicBreak.Label", "Cubic (broken)");
 	case ETangentModeComboState::Mixed: return LOCTEXT("TangentMode.Mixed.Label", "Mixed");
+	case ETangentModeComboState::NoKeys: return LOCTEXT("TangentMode.NoKeys.Label", "No keys");
 	default: checkNoEntry(); return FText::GetEmpty();
 	}
 }
@@ -963,7 +896,6 @@ FText SCurveEditorPanel::GetTangentModeTooltip() const
 {
 	switch (DetermineTangentMode())
 	{
-	case ETangentModeComboState::NoSelection: return LOCTEXT("TangentMode.NoSelection.Tooltip", "Select keys to set tangent mode");
 	case ETangentModeComboState::Constant: return FCurveEditorCommands::Get().InterpolationConstant->GetDescription();
 	case ETangentModeComboState::Linear: return FCurveEditorCommands::Get().InterpolationLinear->GetDescription();
 	case ETangentModeComboState::CubicAuto: return FCurveEditorCommands::Get().InterpolationCubicAuto->GetDescription();
@@ -971,6 +903,7 @@ FText SCurveEditorPanel::GetTangentModeTooltip() const
 	case ETangentModeComboState::CubicUser: return FCurveEditorCommands::Get().InterpolationCubicUser->GetDescription();
 	case ETangentModeComboState::CubicBreak: return FCurveEditorCommands::Get().InterpolationCubicBreak->GetDescription();
 	case ETangentModeComboState::Mixed: return LOCTEXT("TangentMode.Mixed.Tooltip", "The selected keys have different tangent modes.");
+	case ETangentModeComboState::NoKeys: return LOCTEXT("TangentMode.Mixed.Tooltip", "Add some keys to the curve to set their tangent modes.");
 	default: checkNoEntry(); return FText::GetEmpty();
 	}
 }
@@ -979,7 +912,6 @@ FSlateIcon SCurveEditorPanel::GetTangentModeIcon() const
 {
 	switch (DetermineTangentMode())
 	{
-	case ETangentModeComboState::NoSelection: return FSlateIcon(FAppStyle::Get().GetStyleSetName(), "GenericCurveEditor.InterpolationNoSelection");
 	case ETangentModeComboState::Constant: return FSlateIcon(FAppStyle::Get().GetStyleSetName(), "GenericCurveEditor.InterpolationConstant");
 	case ETangentModeComboState::Linear: return FSlateIcon(FAppStyle::Get().GetStyleSetName(), "GenericCurveEditor.InterpolationLinear");
 	case ETangentModeComboState::CubicAuto: return FSlateIcon(FAppStyle::Get().GetStyleSetName(), "GenericCurveEditor.InterpolationCubicAuto");
@@ -987,21 +919,16 @@ FSlateIcon SCurveEditorPanel::GetTangentModeIcon() const
 	case ETangentModeComboState::CubicUser: return FSlateIcon(FAppStyle::Get().GetStyleSetName(), "GenericCurveEditor.InterpolationCubicUser");
 	case ETangentModeComboState::CubicBreak: return FSlateIcon(FAppStyle::Get().GetStyleSetName(), "GenericCurveEditor.InterpolationCubicBreak");
 	case ETangentModeComboState::Mixed: return FSlateIcon(FAppStyle::Get().GetStyleSetName(), "GenericCurveEditor.InterpolationMixed");
+	case ETangentModeComboState::NoKeys: return FSlateIcon(FAppStyle::Get().GetStyleSetName(), "GenericCurveEditor.InterpolationMixed");
 	default: checkNoEntry(); return FSlateIcon();
 	}
 }
 
-bool SCurveEditorPanel::IsTangentModeComboEnabled() const
-{
-	const ETangentModeComboState ComboState = DetermineTangentMode();
-	return ComboState != ETangentModeComboState::NoSelection;
-}
-
 SCurveEditorPanel::ETangentModeComboState SCurveEditorPanel::DetermineTangentMode() const
 {
-	if (CurveEditor->Selection.IsEmpty())
+	if (!CanSetKeyInterpolation())
 	{
-		return ETangentModeComboState::NoSelection;
+		return ETangentModeComboState::NoKeys;
 	}
 	
 	const bool bIsConstantCommon		= CompareCommonInterpolationMode(RCIM_Constant);
@@ -1044,19 +971,47 @@ EVisibility SCurveEditorPanel::GetSplitterVisibility() const
 	return EVisibility::Visible;
 }
 
-void SCurveEditorPanel::SetKeyAttributes(FKeyAttributes KeyAttributes, FText Description)
+void SCurveEditorPanel::HandleSetKeyAttributesCommand(FKeyAttributes InKeyAttributes, FText InDescription)
+{
+	if (CurveEditor->Selection.IsEmpty())
+	{
+		SetKeyAttributesOnAllCurves(InKeyAttributes, InDescription);
+	}
+	else
+	{
+		SetKeyAttributesOnSelection(InKeyAttributes, InDescription);
+	}
+}
+
+void SCurveEditorPanel::SetKeyAttributesOnSelection(const FKeyAttributes& InKeyAttributes, const FText& InDescription) const
 {
 	using namespace UE::CurveEditor;
 	const FScopedCurveChange KeyChange(
-		FCurvesSnapshotBuilder(GetCurveEditor(), ECurveChangeFlags::KeyAttributes).TrackSelectedCurves(),
-		Description
-		);
+		FCurvesSnapshotBuilder(GetCurveEditor(), ECurveChangeFlags::KeyAttributes).TrackSelectedCurves(), InDescription
+	);
 	
 	for (const TTuple<FCurveModelID, FKeyHandleSet>& Pair : CurveEditor->Selection.GetAll())
 	{
 		if (FCurveModel* Curve = CurveEditor->FindCurve(Pair.Key))
 		{
-			Curve->SetKeyAttributes(Pair.Value.AsArray(), KeyAttributes);
+			Curve->SetKeyAttributes(Pair.Value.AsArray(), InKeyAttributes);
+		}
+	}
+}
+
+void SCurveEditorPanel::SetKeyAttributesOnAllCurves(const FKeyAttributes& InKeyAttributes, const FText& InDescription) const
+{
+	using namespace UE::CurveEditor;
+	const FScopedCurveChange KeyChange(
+		FCurvesSnapshotBuilder(GetCurveEditor(), ECurveChangeFlags::KeyAttributes).TrackAllCurves(), InDescription
+	);
+	
+	for (const TPair<FCurveModelID, TUniquePtr<FCurveModel>>& Pair : CurveEditor->GetCurves())
+	{
+		if (FCurveModel* Curve = Pair.Value.Get())
+		{
+			const TArray<FKeyHandle> KeyHandles = Curve->GetAllKeys();
+			Curve->SetKeyAttributes(KeyHandles, InKeyAttributes);
 		}
 	}
 }
@@ -1139,7 +1094,12 @@ bool SCurveEditorPanel::CanToggleWeightedTangents() const
 
 bool SCurveEditorPanel::CanSetKeyInterpolation() const
 {
-	return CurveEditor->GetSelection().Count() > 0;
+	// Basically, return false if there are no keys... if user has selected something, we can set.
+	// If the user has nothing selected, we can set all the keys in view. But if there are no keys, we want to disable the combo button.
+	return Algo::AnyOf(CurveEditor->GetCurves(), [](const TPair<FCurveModelID, TUniquePtr<FCurveModel>>& Pair)
+	{
+		return Pair.Value && Pair.Value->GetNumKeys() > 0;
+	});
 }
 
 FReply SCurveEditorPanel::OnKeyDown(const FGeometry& MyGeometry, const FKeyEvent& InKeyEvent)
@@ -1524,7 +1484,7 @@ void SCurveEditorPanel::BuildToolbar(FToolBarBuilder& InToolBarBuilder, TSharedP
 		InToolBarBuilder.BeginStyleOverride("CurveEditorToolBar.TangentMode"); // Min size for combo button so size stays consistent independent of label size
 		InToolBarBuilder.SetLabelVisibility(EVisibility::Visible); // Show label...
 		InToolBarBuilder.AddComboButton(
-			FUIAction(FExecuteAction(), FCanExecuteAction::CreateSP(this, &SCurveEditorPanel::IsTangentModeComboEnabled)),
+			FUIAction(),
 			FOnGetContent::CreateSP(this, &SCurveEditorPanel::MakeTangentModeMenu),
 			TAttribute<FText>::CreateSP(this, &SCurveEditorPanel::GetTangentModeLabel),
 			TAttribute<FText>::CreateSP(this, &SCurveEditorPanel::GetTangentModeTooltip),
diff --git a/Engine/Source/Editor/CurveEditor/Public/SCurveEditorPanel.h b/Engine/Source/Editor/CurveEditor/Public/SCurveEditorPanel.h
--- a/Engine/Source/Editor/CurveEditor/Public/SCurveEditorPanel.h
+++ b/Engine/Source/Editor/CurveEditor/Public/SCurveEditorPanel.h
@@ -227,12 +227,9 @@ private:
 	UE_API FText GetTangentModeLabel() const;
 	UE_API FText GetTangentModeTooltip() const;
 	UE_API FSlateIcon GetTangentModeIcon() const;
-	UE_API bool IsTangentModeComboEnabled() const;
 	
 	enum class ETangentModeComboState : uint8
 	{
-		NoSelection,
-
 		// All keys in selection share the following mode:
 		Constant,
 		Linear,
@@ -242,7 +239,9 @@ private:
 		CubicBreak,
 
 		// The keys are mixed
-		Mixed
+		Mixed,
+		// There are no keys in view so nothing can be set.
+		NoKeys
 	};
 	UE_API ETangentModeComboState DetermineTangentMode() const;
 	
@@ -256,7 +255,6 @@ private:
 	/*~ Event bindings */
 	UE_API void UpdateTime();
 	UE_API void UpdateEditBox();
-	UE_API void UpdateCommonCurveInfo();
 
 	/** Creates the drop-down list you see when changing Curve View options. */
 	UE_API TSharedRef<SWidget> MakeCurveEditorCurveViewOptionsMenu(TSharedPtr<FExtender> InExtender);
@@ -275,17 +273,19 @@ private:
 	/**
 	 * Bind command mappings for this widget
 	 */
-	UE_API void BindCommands();
+	void BindCommands();
 	
-	/**
-	 * Assign new attributes to the currently selected keys
-	 */
-	UE_API void SetKeyAttributes(FKeyAttributes KeyAttributes, FText Description);
+	/** Assign new attributes to the currently selected keys*/
+	void HandleSetKeyAttributesCommand(FKeyAttributes InKeyAttributes, FText InDescription);
+	/** Sets the key attributes on all selected keys. */
+	void SetKeyAttributesOnSelection(const FKeyAttributes& InKeyAttributes, const FText& InDescription) const;
+	/** Sets the key attributes on all keys. */
+	void SetKeyAttributesOnAllCurves(const FKeyAttributes& InKeyAttributes, const FText& InDescription) const;
 
 	/**
 	 * Assign new curve attributes to all visible curves
 	 */
-	UE_API void SetCurveAttributes(FCurveAttributes CurveAttributes, FText Description);
+	void SetCurveAttributes(FCurveAttributes CurveAttributes, FText Description);
 
 	/** Compare all the currently selected keys' interp modes against the specified interp mode */
 	bool CompareCommonInterpolationMode(ERichCurveInterpMode InterpMode) const
