Fix loca key collision

diff --git a/Engine/Source/Editor/CurveEditor/Private/SCurveEditorPanel.cpp b/Engine/Source/Editor/CurveEditor/Private/SCurveEditorPanel.cpp
--- a/Engine/Source/Editor/CurveEditor/Private/SCurveEditorPanel.cpp
+++ b/Engine/Source/Editor/CurveEditor/Private/SCurveEditorPanel.cpp
@@ -903,7 +903,7 @@ FText SCurveEditorPanel::GetTangentModeTooltip() const
 	case ETangentModeComboState::CubicUser: return FCurveEditorCommands::Get().InterpolationCubicUser->GetDescription();
 	case ETangentModeComboState::CubicBreak: return FCurveEditorCommands::Get().InterpolationCubicBreak->GetDescription();
 	case ETangentModeComboState::Mixed: return LOCTEXT("TangentMode.Mixed.Tooltip", "The selected keys have different tangent modes.");
-	case ETangentModeComboState::NoKeys: return LOCTEXT("TangentMode.Mixed.Tooltip", "Add some keys to the curve to set their tangent modes.");
+	case ETangentModeComboState::NoKeys: return LOCTEXT("TangentMode.NoKeys.Tooltip", "Add some keys to the curve to set their tangent modes.");
 	default: checkNoEntry(); return FText::GetEmpty();
 	}
 }
