Curve Editor: Shrink the amount of memory stored in undo commands. Easy win to save a little bit of memory.

diff --git a/Engine/Source/Editor/CurveEditor/Private/Modification/Keys/GenericCurveChangeCommand.h b/Engine/Source/Editor/CurveEditor/Private/Modification/Keys/GenericCurveChangeCommand.h
--- a/Engine/Source/Editor/CurveEditor/Private/Modification/Keys/GenericCurveChangeCommand.h
+++ b/Engine/Source/Editor/CurveEditor/Private/Modification/Keys/GenericCurveChangeCommand.h
@@ -21,7 +21,8 @@ public:
 
 	explicit FGenericCurveChangeCommand(TSharedPtr<FCurveEditor> InCurveEditor, FGenericCurveChangeData InDeltaChange)
 		: FCurveEditorCommandChange(MoveTemp(InCurveEditor))
-		, DeltaChange(MoveTemp(InDeltaChange))
+		// Make sure to Shrink() whatever is passed to us so it takes up as little memory as possible.
+		, DeltaChange([InDeltaChange = MoveTemp(InDeltaChange)]() mutable { InDeltaChange.Shrink(); return InDeltaChange; }())
 	{}
 
 	//~ Begin FChange Interface
diff --git a/Engine/Source/Editor/CurveEditor/Private/Modification/Selection/SelectionChangeCommand.h b/Engine/Source/Editor/CurveEditor/Private/Modification/Selection/SelectionChangeCommand.h
--- a/Engine/Source/Editor/CurveEditor/Private/Modification/Selection/SelectionChangeCommand.h
+++ b/Engine/Source/Editor/CurveEditor/Private/Modification/Selection/SelectionChangeCommand.h
@@ -14,7 +14,7 @@ public:
 
 	explicit FSelectionChangeCommand(TSharedPtr<FCurveEditor> InCurveEditor, FSelectionDeltaChange InDeltaChange)
 		: FCurveEditorCommandChange(MoveTemp(InCurveEditor))
-		, DeltaChange(MoveTemp(InDeltaChange))
+		, DeltaChange([InDeltaChange = MoveTemp(InDeltaChange)]() mutable { InDeltaChange.Shrink(); return InDeltaChange; }())
 	{}
 	
 	//~ Begin FChange Interface
diff --git a/Engine/Source/Editor/CurveEditor/Private/Modification/Selection/SelectionDeltaChange.h b/Engine/Source/Editor/CurveEditor/Private/Modification/Selection/SelectionDeltaChange.h
--- a/Engine/Source/Editor/CurveEditor/Private/Modification/Selection/SelectionDeltaChange.h
+++ b/Engine/Source/Editor/CurveEditor/Private/Modification/Selection/SelectionDeltaChange.h
@@ -27,6 +27,13 @@ struct FCurveSelectionDeltaChange
 	{
 		return AddedKeys.GetAllocatedSize() + RemovedKeys.GetAllocatedSize();
 	}
+
+	/** Shrinks the containers' used memory to smallest possible to store elements currently in it. */
+	void Shrink()
+	{
+		AddedKeys.Shrink();
+		RemovedKeys.Shrink();
+	}
 };
 
 /** Changes made to the curve editor selection */
@@ -51,5 +58,15 @@ struct FSelectionDeltaChange
 				return Pair.Value.GetAllocatedSize();
 			}, 0);
 	}
+
+	/** Shrinks the containers' used memory to smallest possible to store elements currently in it. */
+	void Shrink()
+	{
+		ChangedCurves.Shrink();
+		for (TPair<FCurveModelID, FCurveSelectionDeltaChange>& Pair : ChangedCurves)
+		{
+			ChangedCurves.Shrink();
+		}
+	}
 };
 }
diff --git a/Engine/Source/Editor/CurveEditor/Public/Modification/Keys/Data/AddKeysChangeData.h b/Engine/Source/Editor/CurveEditor/Public/Modification/Keys/Data/AddKeysChangeData.h
--- a/Engine/Source/Editor/CurveEditor/Public/Modification/Keys/Data/AddKeysChangeData.h
+++ b/Engine/Source/Editor/CurveEditor/Public/Modification/Keys/Data/AddKeysChangeData.h
@@ -17,5 +17,15 @@ struct FAddKeysChangeData
 	bool HasChanges() const { return !SavedCurveState.IsEmpty(); }
 	operator bool() const { return HasChanges(); }
 	SIZE_T GetAllocatedSize() const { return FCurveKeyData::GetAllocatedSize(SavedCurveState); }
+
+	/** Shrinks the containers' used memory to smallest possible to store elements currently in it. */
+	void Shrink()
+	{
+		SavedCurveState.Shrink();
+		for (TPair<FCurveModelID, FCurveKeyData>& Pair : SavedCurveState)
+		{
+			Pair.Value.Shrink();
+		}
+	}
 };
 }
diff --git a/Engine/Source/Editor/CurveEditor/Public/Modification/Keys/Data/CurveKeyDataMap.h b/Engine/Source/Editor/CurveEditor/Public/Modification/Keys/Data/CurveKeyDataMap.h
--- a/Engine/Source/Editor/CurveEditor/Public/Modification/Keys/Data/CurveKeyDataMap.h
+++ b/Engine/Source/Editor/CurveEditor/Public/Modification/Keys/Data/CurveKeyDataMap.h
@@ -50,6 +50,14 @@ struct FCurveKeyData
 			return Pair.Value.GetAllocatedSize();
 		}, 0);
 	}
+
+	/** Shrinks the containers' used memory to smallest possible to store elements currently in it. */
+	void Shrink()
+	{
+		KeyHandles.Shrink();
+		KeyPositions.Shrink();
+		KeyAttributes.Shrink();
+	}
 };
 
 /** Snapshot of curve editor content. */
diff --git a/Engine/Source/Editor/CurveEditor/Public/Modification/Keys/Data/GenericCurveChangeData.h b/Engine/Source/Editor/CurveEditor/Public/Modification/Keys/Data/GenericCurveChangeData.h
--- a/Engine/Source/Editor/CurveEditor/Public/Modification/Keys/Data/GenericCurveChangeData.h
+++ b/Engine/Source/Editor/CurveEditor/Public/Modification/Keys/Data/GenericCurveChangeData.h
@@ -39,5 +39,14 @@ struct FGenericCurveChangeData
 			+ KeyAttributeData.GetAllocatedSize()
 			+ CurveAttributeData.GetAllocatedSize();
 	}
+	
+	/** Shrinks the containers' used memory to smallest possible to store elements currently in it. */
+	void Shrink()
+	{
+		MoveKeysData.Shrink();
+		AddKeysData.Shrink();
+		RemoveKeysData.Shrink();
+		KeyAttributeData.Shrink();
+	}
 };
 }
diff --git a/Engine/Source/Editor/CurveEditor/Public/Modification/Keys/Data/KeyAttributeChangeData.h b/Engine/Source/Editor/CurveEditor/Public/Modification/Keys/Data/KeyAttributeChangeData.h
--- a/Engine/Source/Editor/CurveEditor/Public/Modification/Keys/Data/KeyAttributeChangeData.h
+++ b/Engine/Source/Editor/CurveEditor/Public/Modification/Keys/Data/KeyAttributeChangeData.h
@@ -29,6 +29,14 @@ struct FKeyAttributeChangeData_PerCurve
 	bool HasChanges() const { return !Handles.IsEmpty(); }
 	operator bool() const { return HasChanges(); }
 	SIZE_T GetAllocatedSize() const { return Handles.GetAllocatedSize() + BeforeChange.GetAllocatedSize() + AfterChange.GetAllocatedSize(); }
+
+	/** Shrinks the containers' used memory to smallest possible to store elements currently in it. */
+	void Shrink()
+	{
+		Handles.Shrink();
+		BeforeChange.Shrink();
+		AfterChange.Shrink();
+	}
 };
 
 /** Stores data for restoring key attributes of multiple curves. */
@@ -46,5 +54,15 @@ struct FKeyAttributeChangeData
 				return Pair.Value.GetAllocatedSize();
 			}, 0);
 	}
+
+	/** Shrinks the containers' used memory to smallest possible to store elements currently in it. */
+	void Shrink()
+	{
+		ChangedCurves.Shrink();
+		for (TPair<FCurveModelID, FKeyAttributeChangeData_PerCurve>& Pair : ChangedCurves)
+		{
+			Pair.Value.Shrink();
+		}
+	}
 };
 }
diff --git a/Engine/Source/Editor/CurveEditor/Public/Modification/Keys/Data/MoveKeysChangeData.h b/Engine/Source/Editor/CurveEditor/Public/Modification/Keys/Data/MoveKeysChangeData.h
--- a/Engine/Source/Editor/CurveEditor/Public/Modification/Keys/Data/MoveKeysChangeData.h
+++ b/Engine/Source/Editor/CurveEditor/Public/Modification/Keys/Data/MoveKeysChangeData.h
@@ -25,6 +25,13 @@ struct FMoveKeysChangeData_PerCurve
 	bool HasChanges() const { return !Handles.IsEmpty(); }
 	operator bool() const { return HasChanges(); }
 	SIZE_T GetAllocatedSize() const { return Handles.GetAllocatedSize() + Translations.GetAllocatedSize(); }
+
+	/** Shrinks the containers' used memory to smallest possible to store elements currently in it. */
+	void Shrink()
+	{
+		Handles.Shrink();
+		Translations.Shrink();
+	}
 };
 
 /**
@@ -46,5 +53,15 @@ struct FMoveKeysChangeData
 				return Pair.Value.GetAllocatedSize();
 			}, 0);
 	}
+
+	/** Shrinks the containers' used memory to smallest possible to store elements currently in it. */
+	void Shrink()
+	{
+		ChangedCurves.Shrink();
+		for (TPair<FCurveModelID, FMoveKeysChangeData_PerCurve>& Pair : ChangedCurves)
+		{
+			Pair.Value.Shrink();
+		}
+	}
 };
 }
diff --git a/Engine/Source/Editor/CurveEditor/Public/Modification/Keys/Data/RemoveKeysChangeData.h b/Engine/Source/Editor/CurveEditor/Public/Modification/Keys/Data/RemoveKeysChangeData.h
--- a/Engine/Source/Editor/CurveEditor/Public/Modification/Keys/Data/RemoveKeysChangeData.h
+++ b/Engine/Source/Editor/CurveEditor/Public/Modification/Keys/Data/RemoveKeysChangeData.h
@@ -19,5 +19,15 @@ struct FRemoveKeysChangeData
 	bool HasChanges() const { return !SavedCurveState.IsEmpty(); }
 	operator bool() const { return HasChanges(); }
 	SIZE_T GetAllocatedSize() const { return FCurveKeyData::GetAllocatedSize(SavedCurveState); }
+
+	/** Shrinks the containers' used memory to smallest possible to store elements currently in it. */
+	void Shrink()
+	{
+		SavedCurveState.Shrink();
+		for (TPair<FCurveModelID, FCurveKeyData>& Pair : SavedCurveState)
+		{
+			Pair.Value.Shrink();
+		}
+	}
 };
 }
