Curve Editor: When moving keys using mouse drag, we now generate 1 transaction instead of 2:
- LMB click key & move keys
- MMB insert a key, keep MMB pressed, and drag to position.

diff --git a/Engine/Source/Editor/CurveEditor/Private/DragOperations/CurveEditorDragOperation_MoveKeys.cpp b/Engine/Source/Editor/CurveEditor/Private/DragOperations/CurveEditorDragOperation_MoveKeys.cpp
--- a/Engine/Source/Editor/CurveEditor/Private/DragOperations/CurveEditorDragOperation_MoveKeys.cpp
+++ b/Engine/Source/Editor/CurveEditor/Private/DragOperations/CurveEditorDragOperation_MoveKeys.cpp
@@ -34,17 +34,8 @@ void FCurveEditorDragOperation_MoveKeys::OnInitialize(FCurveEditor* InCurveEdito
 void FCurveEditorDragOperation_MoveKeys::OnBeginDrag(FVector2D InitialPosition, FVector2D CurrentPosition, const FPointerEvent& MouseEvent)
 {
 	using namespace UE::CurveEditor;
-
-	// In OnBeginDrag, we should already be in a transaction. Setting "Transaction" will extend that transaction until the user lets go.
-	const int32 NumKeys = CurveEditor->GetSelection().Count();
-	Transaction = MakeUnique<FScopedTransaction>(
-		FText::Format(NSLOCTEXT("CurveEditor", "MoveKeysFormat", "Move {0}|plural(one=Key, other=Keys)"), NumKeys)
-		);
 	
-	TSharedRef<FCurveEditor> CurveEditorPin = CurveEditor->AsShared();
-	ScopedKeyChange.Emplace(
-		FCurvesSnapshotBuilder(CurveEditorPin, ECurveChangeFlags::MoveKeys).TrackSelectedKeys()
-		);
+	ScopedKeyChange.Emplace(FCurvesSnapshotBuilder(CurveEditor->AsShared(), ECurveChangeFlags::MoveKeys).TrackSelectedKeys());
 
 	KeysByCurve.Reset();
 	CurveEditor->SuppressBoundTransformUpdates(true);
@@ -114,7 +105,6 @@ void FCurveEditorDragOperation_MoveKeys::OnCancelDrag()
 	
 	ScopedKeyChange->Cancel();
 	ScopedKeyChange.Reset();
-	Transaction.Reset();
 }
 
 void FCurveEditorDragOperation_MoveKeys::OnEndDrag(FVector2D InitialPosition, FVector2D CurrentPosition, const FPointerEvent& MouseEvent)
@@ -132,7 +122,6 @@ void FCurveEditorDragOperation_MoveKeys::OnEndDrag(FVector2D InitialPosition, FV
 	CurveEditor->SuppressBoundTransformUpdates(false);
 	
 	ScopedKeyChange.Reset();
-	Transaction.Reset();
 }
 
 FCurveEditorDragOperation_MoveKeys::FAccumulatedMouseMovement& FCurveEditorDragOperation_MoveKeys::GetOrAddAccumulatedMouseMovement(
diff --git a/Engine/Source/Editor/CurveEditor/Private/DragOperations/CurveEditorDragOperation_Tangent.cpp b/Engine/Source/Editor/CurveEditor/Private/DragOperations/CurveEditorDragOperation_Tangent.cpp
--- a/Engine/Source/Editor/CurveEditor/Private/DragOperations/CurveEditorDragOperation_Tangent.cpp
+++ b/Engine/Source/Editor/CurveEditor/Private/DragOperations/CurveEditorDragOperation_Tangent.cpp
@@ -230,17 +230,8 @@ void FCurveEditorDragOperation_Tangent::OnInitialize(FCurveEditor* InCurveEditor
 void FCurveEditorDragOperation_Tangent::OnBeginDrag(FVector2D InitialPosition, FVector2D CurrentPosition, const FPointerEvent& MouseEvent)
 {
 	using namespace UE::CurveEditor;
-
-	// In OnBeginDrag, we should already be in a transaction. Setting "Transaction" will extend that transaction until the user lets go.
-	const int32 NumKeys = CurveEditor->GetSelection().Count();
-	Transaction = MakeUnique<FScopedTransaction>(
-		FText::Format(NSLOCTEXT("CurveEditor", "DragTangentsFormat", "Drag {0}|plural(one=Tangent, other=Tangents)"), NumKeys)
-		);
 	
-	TSharedRef<FCurveEditor> CurveEditorPin = CurveEditor->AsShared();
-	ScopedKeyChange.Emplace(
-		FCurvesSnapshotBuilder(CurveEditorPin, ECurveChangeFlags::KeyAttributes).TrackSelectedCurves()
-		);
+	ScopedKeyChange.Emplace(FCurvesSnapshotBuilder(CurveEditor->AsShared(), ECurveChangeFlags::KeyAttributes).TrackSelectedCurves());
 	CurveEditor->SuppressBoundTransformUpdates(true);
 
 	KeysByCurve.Reset();
@@ -329,7 +320,6 @@ void FCurveEditorDragOperation_Tangent::OnCancelDrag()
 	
 	ScopedKeyChange->Cancel();
 	ScopedKeyChange.Reset();
-	Transaction.Reset();
 }
 
 void FCurveEditorDragOperation_Tangent::OnEndDrag(FVector2D InitialPosition, FVector2D CurrentPosition, const FPointerEvent& MouseEvent)
@@ -350,5 +340,4 @@ void FCurveEditorDragOperation_Tangent::OnEndDrag(FVector2D InitialPosition, FVe
 	
 	CurveEditor->SuppressBoundTransformUpdates(false);
 	ScopedKeyChange.Reset();
-	Transaction.Reset();
 }
diff --git a/Engine/Source/Editor/CurveEditor/Private/DragOperations/KeyDragOperationUtils.cpp b/Engine/Source/Editor/CurveEditor/Private/DragOperations/KeyDragOperationUtils.cpp
new file mode 100644
--- /dev/null
+++ b/Engine/Source/Editor/CurveEditor/Private/DragOperations/KeyDragOperationUtils.cpp
@@ -0,0 +1,68 @@
+﻿// Copyright Epic Games, Inc. All Rights Reserved.
+
+#include "KeyDragOperationUtils.h"
+
+#include "CurveEditor.h"
+#include "CurveEditorDragOperation_MoveKeys.h"
+#include "CurveEditorDragOperation_Tangent.h"
+#include "CurveEditorTypes.h"
+
+namespace UE::CurveEditor
+{
+TUniquePtr<ICurveEditorKeyDragOperation> CreateAndInitializeKeyDrag(
+	FCurveEditor& InCurveEditor, ECurvePointType InKeyType, const TOptional<FCurvePointHandle>& InCardinalPoint
+	)
+{
+	TUniquePtr<ICurveEditorKeyDragOperation> Result;
+	switch (InKeyType)
+	{
+	case ECurvePointType::ArriveTangent:
+	case ECurvePointType::LeaveTangent:
+		Result = MakeUnique<FCurveEditorDragOperation_Tangent>();
+		break;
+
+	default:
+		Result = MakeUnique<FCurveEditorDragOperation_MoveKeys>();
+	}
+
+	if (Result)
+	{
+		Result->Initialize(&InCurveEditor, InCardinalPoint);
+	}
+	
+	return Result;
+}
+
+namespace Private
+{
+static FText CreateTransactionTitle(ECurvePointType InKeyType, const int32 InNumKeys)
+{
+	switch (InKeyType)
+	{
+	case ECurvePointType::ArriveTangent:
+	case ECurvePointType::LeaveTangent:
+		return FText::Format(NSLOCTEXT("CurveEditor", "DragTangentsFormat", "Drag {0}|plural(one=Tangent, other=Tangents)"), InNumKeys);
+
+	default:
+		return FText::Format(NSLOCTEXT("CurveEditor", "MoveKeysFormat", "Move {0}|plural(one=Key, other=Keys)"), InNumKeys);
+	}
+}
+}
+
+TUniquePtr<FScopedTransaction> CreateKeyOperationTransaction(ECurvePointType InKeyType, const int32 InNumKeys)
+{
+	// For now all operations create a transaction. In the future, you might add an operation that requires no transactions.
+	// In that case, it's fine if you introduce some control flow here to return nullptr.
+	return MakeUnique<FScopedTransaction>(Private::CreateTransactionTitle(InKeyType, InNumKeys));
+}
+
+TUniquePtr<FScopedTransaction> CreateKeyOperationTransaction(const FCurveEditor& InCurveEditor, ECurvePointType InKeyType)
+{
+	int32 NumKeys = 0;
+	for (const TPair<FCurveModelID, FKeyHandleSet>& Pair : InCurveEditor.GetSelection().GetAll())
+	{
+		NumKeys += Pair.Value.Num();
+	}
+	return CreateKeyOperationTransaction(InKeyType, NumKeys);
+}
+}
diff --git a/Engine/Source/Editor/CurveEditor/Private/DragOperations/KeyDragOperationUtils.h b/Engine/Source/Editor/CurveEditor/Private/DragOperations/KeyDragOperationUtils.h
new file mode 100644
--- /dev/null
+++ b/Engine/Source/Editor/CurveEditor/Private/DragOperations/KeyDragOperationUtils.h
@@ -0,0 +1,26 @@
+﻿// Copyright Epic Games, Inc. All Rights Reserved.
+
+#pragma once
+
+#include "HAL/Platform.h"
+#include "Misc/Optional.h"
+#include "Templates/UniquePtr.h"
+
+class FCurveEditor;
+class FScopedTransaction;
+class ICurveEditorKeyDragOperation;
+enum class ECurvePointType : uint8;
+struct FCurvePointHandle;
+
+namespace UE::CurveEditor
+{
+/** Creates the drag operation based on InKeyType. */
+TUniquePtr<ICurveEditorKeyDragOperation> CreateAndInitializeKeyDrag(
+	FCurveEditor& InCurveEditor, ECurvePointType InKeyType, const TOptional<FCurvePointHandle>& InCardinalPoint
+	);
+
+/** Creates the transaction that the ICurveEditorKeyDragOperation created by CreateKeyDrag would use. */
+TUniquePtr<FScopedTransaction> CreateKeyOperationTransaction(ECurvePointType InKeyType, const int32 InNumKeys);
+/** Version that determines the number of selected keys. */
+TUniquePtr<FScopedTransaction> CreateKeyOperationTransaction(const FCurveEditor& InCurveEditor, ECurvePointType InKeyType);
+}
diff --git a/Engine/Source/Editor/CurveEditor/Private/SCurveEditorViewContainer.cpp b/Engine/Source/Editor/CurveEditor/Private/SCurveEditorViewContainer.cpp
--- a/Engine/Source/Editor/CurveEditor/Private/SCurveEditorViewContainer.cpp
+++ b/Engine/Source/Editor/CurveEditor/Private/SCurveEditorViewContainer.cpp
@@ -200,7 +200,7 @@ FReply SCurveEditorViewContainer::OnKeyDown(const FGeometry& MyGeometry, const F
 	if (!ToolReply.IsEventHandled() &&
 		InKeyEvent.GetKey() == EKeys::Escape && DragOperation.IsSet())
 	{
-		DragOperation->DragImpl->CancelDrag();
+		DragOperation->Cancel();
 		DragOperation.Reset();
 		return FReply::Handled();
 	}
@@ -229,7 +229,7 @@ FReply SCurveEditorViewContainer::OnKeyUp(const FGeometry& MyGeometry, const FKe
 		bIsScrubbingTime = false;
 		if (DragOperation.IsSet())
 		{
-			DragOperation->DragImpl->CancelDrag();
+			DragOperation->Cancel();
 			DragOperation.Reset();
 			return FReply::Handled();
 		}
@@ -414,7 +414,7 @@ void SCurveEditorViewContainer::OnMouseCaptureLost(const FCaptureLostEvent& Capt
 		{
 			if (!HasMouseCapture())
 			{
-				DragOperation->DragImpl->CancelDrag();
+				DragOperation->Cancel();
 				DragOperation.Reset();
 			}
 		});
@@ -437,7 +437,7 @@ void SCurveEditorViewContainer::OnFocusLost(const FFocusEvent& InFocusEvent)
 {
 	if (DragOperation.IsSet())
 	{
-		DragOperation->DragImpl->CancelDrag();
+		DragOperation->Cancel();
 		DragOperation.Reset();
 	}
 
@@ -457,7 +457,7 @@ void SCurveEditorViewContainer::OnCurveEditorToolChanged(FCurveEditorToolID InTo
 	if (DragOperation.IsSet())
 	{
 		// We have to cancel it instead of ending it because ending it needs mouse position and some other stuff.
-		DragOperation->DragImpl->CancelDrag();
+		DragOperation->Cancel();
 		DragOperation.Reset();
 	}
 }
diff --git a/Engine/Source/Editor/CurveEditor/Private/Views/SInteractiveCurveEditorView.cpp b/Engine/Source/Editor/CurveEditor/Private/Views/SInteractiveCurveEditorView.cpp
--- a/Engine/Source/Editor/CurveEditor/Private/Views/SInteractiveCurveEditorView.cpp
+++ b/Engine/Source/Editor/CurveEditor/Private/Views/SInteractiveCurveEditorView.cpp
@@ -63,6 +63,7 @@
 #include "SCurveEditorPanel.h"
 #include "ScopedTransaction.h"
 #include "SlotBase.h"
+#include "DragOperations/KeyDragOperationUtils.h"
 #include "Modification/Utils/ScopedSelectionChange.h"
 #include "Styling/AppStyle.h"
 #include "Styling/CoreStyle.h"
@@ -100,19 +101,6 @@ extern UNREALED_API UEditorEngine* GEditor;
 
 #define LOCTEXT_NAMESPACE "SInteractiveCurveEditorView"
 
-TUniquePtr<ICurveEditorKeyDragOperation> CreateKeyDrag(ECurvePointType KeyType)
-{
-	switch (KeyType)
-	{
-	case ECurvePointType::ArriveTangent:
-	case ECurvePointType::LeaveTangent:
-		return MakeUnique<FCurveEditorDragOperation_Tangent>();
-
-	default:
-		return MakeUnique<FCurveEditorDragOperation_MoveKeys>();
-	}
-}
-
 class SDynamicToolTip : public SToolTip
 {
 public:
@@ -1473,8 +1461,7 @@ double SInteractiveCurveEditorView::GetTangentValue(const double InTime, const d
 
 void SInteractiveCurveEditorView::HandleDirectKeySelectionByMouse(TSharedPtr<FCurveEditor> CurveEditor, const FPointerEvent& MouseEvent, TOptional<FCurvePointHandle> MouseDownPoint)
 {
-	const UE::CurveEditor::FScopedSelectionChange Transaction(WeakCurveEditor, LOCTEXT("ClickKey", "Click Key"));
-	if(!MouseDownPoint.IsSet())
+	if (!MouseDownPoint.IsSet())
 	{
 		CurveEditor->GetSelection().Clear();
 		return;
@@ -1605,11 +1592,14 @@ FReply SInteractiveCurveEditorView::OnMouseButtonDown(const FGeometry& MyGeometr
 	// Rebind our context actions so that shift click commands use the right position.
 	RebindContextualActions(MousePixel);
 
+
 	if (CurveEditor->GetSettings()->AllowMouseEdit(MouseEvent))
 	{
-		// Middle Mouse can try to create keys on curves.
+		// Middle Mouse can try to create keys on curves...
 		TOptional<FCurvePointHandle> NewPoint;
-
+		// ... and if MMB causes insertion of a key, the drag operation should append its changes to that same transaction.
+		TUniquePtr<FScopedTransaction> RootTransaction;
+		
 		// Add a key to the closest curve to the mouse
 		if (TOptional<FCurveModelID> HoveredCurve = GetHoveredCurve())
 		{
@@ -1622,7 +1612,8 @@ FReply SInteractiveCurveEditorView::OnMouseButtonDown(const FGeometry& MyGeometr
 			FCurveModel* CurveToAddTo = CurveEditor->FindCurve(HoveredCurve.GetValue());
 			if (CurveToAddTo && !CurveToAddTo->IsReadOnly() && !bDraggingKeys && !bKeyAlreadyExists)
 			{
-				FScopedTransaction Transaction(LOCTEXT("InsertKey", "Insert Key")); // Selection and key change should be part of same undo / redo op.
+				// Selection and key change should be part of same undo / redo op.
+				RootTransaction = MakeUnique<FScopedTransaction>(LOCTEXT("InsertAndMove", "Insert and Move Key"));
 				const FScopedSelectionChange SelectionChange(CurveEditor);
 				const FScopedCurveChange KeyChange(
 					FCurvesSnapshotBuilder(CurveEditor, *HoveredCurve, ECurveChangeFlags::KeyAttributes | ECurveChangeFlags::AddKeys)
@@ -1678,7 +1669,8 @@ FReply SInteractiveCurveEditorView::OnMouseButtonDown(const FGeometry& MyGeometr
 				}
 				else
 				{
-					Transaction.Cancel();
+					RootTransaction->Cancel();
+					RootTransaction.Reset();
 				}
 			}
 		}
@@ -1699,13 +1691,12 @@ FReply SInteractiveCurveEditorView::OnMouseButtonDown(const FGeometry& MyGeometr
 				}
 			}
 		}
-
-		TUniquePtr<ICurveEditorKeyDragOperation> KeyDrag = CreateKeyDrag(PointType);
-
-		KeyDrag->Initialize(CurveEditor.Get(), NewPoint);
-
+		
 		DragOperation = FCurveEditorDelayedDrag(MousePixel, MouseEvent.GetEffectingButton());
-		DragOperation->DragImpl = MoveTemp(KeyDrag);
+		DragOperation->DragImpl = CreateAndInitializeKeyDrag(*CurveEditor, PointType, NewPoint);
+		// If MMB caused a key to be inserted above, append the drag operation changes to that transaction. 
+		// Otherwise, we'll create a new transaction, whose title depends on whether a key or tangents are selected atm.
+		DragOperation->Transaction = RootTransaction ? MoveTemp(RootTransaction) : CreateKeyOperationTransaction(*CurveEditor, PointType);
 
 		return FReply::Handled().PreventThrottling();
 	}
@@ -1730,16 +1721,18 @@ FReply SInteractiveCurveEditorView::OnMouseButtonDown(const FGeometry& MyGeometr
 			{
 				if (!CurveModel->IsReadOnly())
 				{
+					// Intention: The "Click Key" action must be part of the same transaction as the change made by the operation
+					// First, we select the keys...
+					const FScopedSelectionChange Transaction(WeakCurveEditor, LOCTEXT("ClickKey", "Click Key"));
 					HandleDirectKeySelectionByMouse(CurveEditor, MouseEvent, MouseDownPoint);
 
-					TUniquePtr<ICurveEditorKeyDragOperation> KeyDrag = CreateKeyDrag(MouseDownPoint->PointType);
-
-					KeyDrag->Initialize(CurveEditor.Get(), MouseDownPoint);
-
 					DragOperation = FCurveEditorDelayedDrag(MousePixel, MouseEvent.GetEffectingButton());
-					DragOperation->DragImpl = MoveTemp(KeyDrag);
+					DragOperation->DragImpl = CreateAndInitializeKeyDrag(*CurveEditor, MouseDownPoint->PointType, MouseDownPoint);
+					// ... then we'll create the transaction, whose title depends on the number of keys selected...
+					DragOperation->Transaction = CreateKeyOperationTransaction(*CurveEditor, MouseDownPoint->PointType);
 
 					return FReply::Handled().PreventThrottling();
+					// ... and finally, ~FScopedSelectionChange appends its command to the currently open transaction.
 				}
 			}
 		}
@@ -1846,6 +1839,7 @@ FReply SInteractiveCurveEditorView::OnMouseButtonUp(const FGeometry& MyGeometry,
 			{
 				if (!CurveModel->IsReadOnly())
 				{
+					const UE::CurveEditor::FScopedSelectionChange Transaction(WeakCurveEditor, LOCTEXT("ClickKey", "Click Key"));
 					HandleDirectKeySelectionByMouse(CurveEditor, MouseEvent, MouseDownPoint);
 				}
 			}
@@ -2231,7 +2225,7 @@ void SInteractiveCurveEditorView::OnCurveEditorToolChanged(FCurveEditorToolID In
 	if (DragOperation.IsSet())
 	{
 		// We have to cancel it instead of ending it because ending it needs mouse position and some other stuff.
-		DragOperation->DragImpl->CancelDrag();
+		DragOperation->Cancel();
 		DragOperation.Reset();
 	}
 }
diff --git a/Engine/Source/Editor/CurveEditor/Public/ICurveEditorDragOperation.h b/Engine/Source/Editor/CurveEditor/Public/ICurveEditorDragOperation.h
--- a/Engine/Source/Editor/CurveEditor/Public/ICurveEditorDragOperation.h
+++ b/Engine/Source/Editor/CurveEditor/Public/ICurveEditorDragOperation.h
@@ -126,7 +126,6 @@ public:
 
 	/**
 	 * Initialize this drag operation using the specified curve editor pointer and an optional cardinal point
-
 	 * @param InCurveEditor       Curve editor pointer. Guaranteed to persist for the lifetime of this drag.
 	 * @param CardinalPoint       The point that should be considered the origin of this drag.
 	 */
@@ -135,19 +134,7 @@ public:
 protected:
 
 	/** Implementation method for derived types to initialize a drag */
-	virtual void OnInitialize(FCurveEditor* InCurveEditor, const TOptional<FCurvePointHandle>& CardinalPoint)
-	{}
-
-	virtual void OnCancelDrag() override
-	{
-		if (Transaction.IsValid())
-		{
-			Transaction->Cancel();
-		}
-	}
-
-	/** Scoped transaction pointer */
-	TUniquePtr<FScopedTransaction> Transaction;
+	virtual void OnInitialize(FCurveEditor* InCurveEditor, const TOptional<FCurvePointHandle>& CardinalPoint) {}
 };
 
 /**
@@ -155,10 +142,18 @@ protected:
  */
 struct FCurveEditorDelayedDrag : FDelayedDrag
 {
+	/** The drag implementation to use once the drag has started */
+	TUniquePtr<ICurveEditorDragOperation> DragImpl;
+
 	/**
-	 * The drag implementation to use once the drag has started
+	 * Transaction associated with the operation. Only created for certain operations as not all operations need to transact.
+	 * 
+	 * Depending on the operation, this transaction may contain more changes than those created by the drag operation.
+	 * For example, when clicking a curve with middle mouse button, we may insert a key first and then start moving that key using the drag operation;
+	 * the transaction would contain the key insertion and selection change and the drag operation may append a change for moving the inserted key.
+	 * This is the reason that the ICurveEditorDragOperation is not responsible for creating the FScopedTransaction.
 	 */
-	TUniquePtr<ICurveEditorDragOperation> DragImpl;
+	TUniquePtr<FScopedTransaction> Transaction;
 
 	/**
 	 * Start a delayed drag operation at the specified position and effective key
@@ -168,4 +163,20 @@ struct FCurveEditorDelayedDrag : FDelayedDrag
 	{
 		SetTriggerScaleFactor(0.1f);
 	}
+
+	/** Cancels the drag operation and transaction. */
+	void Cancel()
+	{
+		if (DragImpl)
+		{
+			DragImpl->CancelDrag();
+			DragImpl.Reset();
+		}
+
+		if (Transaction)
+		{
+			Transaction->Cancel();
+			Transaction.Reset();
+		}
+	}
 };
