Curve Editor: Fix range-based loop shrinking in for loop - which is a typ bug found BugHawk.

diff --git a/Engine/Source/Editor/CurveEditor/Private/Modification/Selection/SelectionDeltaChange.h b/Engine/Source/Editor/CurveEditor/Private/Modification/Selection/SelectionDeltaChange.h
--- a/Engine/Source/Editor/CurveEditor/Private/Modification/Selection/SelectionDeltaChange.h
+++ b/Engine/Source/Editor/CurveEditor/Private/Modification/Selection/SelectionDeltaChange.h
@@ -65,7 +65,7 @@ struct FSelectionDeltaChange
 		ChangedCurves.Shrink();
 		for (TPair<FCurveModelID, FCurveSelectionDeltaChange>& Pair : ChangedCurves)
 		{
-			ChangedCurves.Shrink();
+			Pair.Value.Shrink();
 		}
 	}
 };
