Curve Editor: Fix CurveEditor.LogDiffPerfStats logging incorrect undo size.

diff --git a/Engine/Source/Editor/CurveEditor/Private/Modification/Keys/Diff/MultiGenericCommandAppender.cpp b/Engine/Source/Editor/CurveEditor/Private/Modification/Keys/Diff/MultiGenericCommandAppender.cpp
--- a/Engine/Source/Editor/CurveEditor/Private/Modification/Keys/Diff/MultiGenericCommandAppender.cpp
+++ b/Engine/Source/Editor/CurveEditor/Private/Modification/Keys/Diff/MultiGenericCommandAppender.cpp
@@ -53,9 +53,11 @@ void FMultiGenericCommandAppender::PostProcessChanges()
 	// There is no package so we don't need to worry about any dirty flags.
 	if (ChangesToCurvesWithoutUObject)
 	{
-		CurveEditor->GetTransactionManager()->AppendChange(
-			MakeUnique<FGenericCurveChangeCommand>(CurveEditor, MoveTemp(ChangesToCurvesWithoutUObject))
+		TUniquePtr<FGenericCurveChangeCommand> PerCurveCommand = MakeUnique<FGenericCurveChangeCommand>(
+			CurveEditor, MoveTemp(ChangesToCurvesWithoutUObject)
 			);
+		AllocationSize += PerCurveCommand->GetSize();
+		CurveEditor->GetTransactionManager()->AppendChange(MoveTemp(PerCurveCommand));
 	}
 }
 
@@ -69,17 +71,17 @@ void FMultiGenericCommandAppender::ProcessDiff(const FCurveModelID& InCurveId, F
 		return;
 	}
 
+	InData.Shrink();
+	
 	if (UObject* OwningObject = CurveModel->GetOwningObject())
 	{
-		CurveEditor->GetTransactionManager()->AppendChange(
-			OwningObject, MakeUnique<FGenericCurveChangeCommand>(CurveEditor, MoveTemp(InData))
-			);
+		TUniquePtr<FGenericCurveChangeCommand> PerCurveCommand = MakeUnique<FGenericCurveChangeCommand>(CurveEditor, MoveTemp(InData));
+		AllocationSize += PerCurveCommand->GetSize();
+		CurveEditor->GetTransactionManager()->AppendChange(OwningObject, MoveTemp(PerCurveCommand));
 	}
 	else
 	{
 		Private::AppendChanges(InCurveId, ChangesToCurvesWithoutUObject, MoveTemp(InData));
 	}
-	
-	AllocationSize += ChangesToCurvesWithoutUObject.GetAllocatedSize();
 }
 }
