Curve Editor: Fix clicking into the editor when you have nothing selected creates "Clear Selection" transaction

diff --git a/Engine/Source/Editor/CurveEditor/Private/Modification/Utils/ScopedSelectionChange.cpp b/Engine/Source/Editor/CurveEditor/Private/Modification/Utils/ScopedSelectionChange.cpp
--- a/Engine/Source/Editor/CurveEditor/Private/Modification/Utils/ScopedSelectionChange.cpp
+++ b/Engine/Source/Editor/CurveEditor/Private/Modification/Utils/ScopedSelectionChange.cpp
@@ -42,6 +42,11 @@ FScopedSelectionChange::~FScopedSelectionChange()
 	}
 
 	FSelectionDeltaChange DeltaChange = DiffSelection(OriginalSelection, CurveEditorPin->Selection);
+	if (!DeltaChange)
+	{
+		return;
+	}
+	
 	if (!bHasBeenCancelled)
 	{
 		TryAppendCommand(MakeUnique<FSelectionChangeCommand>(CurveEditorPin, MoveTemp(DeltaChange)));
