Curve Editor: Enable the new command-based undo system by default. Users can set "CurveEditor.UseNewUndoSystem 0" if it causes problems.

diff --git a/Engine/Source/Editor/CurveEditor/Private/Modification/Utils/ScopedCurveChange.cpp b/Engine/Source/Editor/CurveEditor/Private/Modification/Utils/ScopedCurveChange.cpp
--- a/Engine/Source/Editor/CurveEditor/Private/Modification/Utils/ScopedCurveChange.cpp
+++ b/Engine/Source/Editor/CurveEditor/Private/Modification/Utils/ScopedCurveChange.cpp
@@ -26,7 +26,7 @@ static TAutoConsoleVariable<bool> CVarLogUndoKeyChangeSize(
 	);
 
 static TAutoConsoleVariable<bool> CVarUseNewUndoSystem(
-	TEXT("CurveEditor.UseNewUndoSystem"), false, TEXT("Whether to use the new undo system. Warning: There may still be some bugs.")
+	TEXT("CurveEditor.UseNewUndoSystem"), true, TEXT("Whether to use the new undo system. Warning: There may still be some bugs.")
 	);
 
 static void ModifyCurves(const TSharedPtr<FCurveEditor>& InCurveEditorPin, const FCurveChangeDiff& InChangeDiff)
