Curve Editor: Fix paste relative pasting absolute if curve has no keys in front of scrubber

diff --git a/Engine/Source/Editor/CurveEditor/Private/CurveEditor.cpp b/Engine/Source/Editor/CurveEditor/Private/CurveEditor.cpp
--- a/Engine/Source/Editor/CurveEditor/Private/CurveEditor.cpp
+++ b/Engine/Source/Editor/CurveEditor/Private/CurveEditor.cpp
@@ -1712,15 +1712,25 @@ static double FindRelativeKeyPasteInset(const UCurveEditorCopyableCurveKeys* InS
 
 	TOptional<FKeyHandle> ClosestPrevious, ClosestNext;
 	InTargetCurveModel->GetClosestKeysTo(InCurrentTime, ClosestPrevious, ClosestNext);
-	if (!ClosestPrevious)
+	
+	// If any key is at or in front of the scrubber, paste relative to that key.
+	// If there is no previous key, then the user would expect the first key to be pasted where the scrubber intersects the curve.
+	double KeyOutputValue;
+	if (ClosestPrevious)
+	{		
+		FKeyPosition Position{};
+		InTargetCurveModel->GetKeyPositions(
+			TConstArrayView<FKeyHandle>(&ClosestPrevious.GetValue(), 1), TArrayView<FKeyPosition>(&Position, 1)
+			);
+		KeyOutputValue = Position.OutputValue;
+	}
+	else
 	{
-		// No previous key -> no inset to apply
-		return 0.0;
+		InTargetCurveModel->Evaluate(InCurrentTime, KeyOutputValue);
 	}
-
-	FKeyPosition Position{};
-	InTargetCurveModel->GetKeyPositions(TConstArrayView<FKeyHandle>(&ClosestPrevious.GetValue(), 1), TArrayView<FKeyPosition>(&Position, 1));
-	return Position.OutputValue - MinElement->OutputValue; // Bring all values down to the closest value to the left.
+	
+	// Bring all values down to the closest value to the left.
+	return KeyOutputValue - MinElement->OutputValue; 
 }
 }
 
