Curve Editor: Introduce event FCurveEditorSelection::OnSelectionChanged:
- The event is broadcast whenever any change is made to FCurveEditorSelection.
- Can be suppressed using FScopedSelectionChangeEventSuppression (optimization)
- FScopedSelectionChange suppresses the event until the scope ends.
- FSelectionChangeCommand suppresses the event until after the undo / redo data has been applied.

diff --git a/Engine/Source/Editor/CurveEditor/Private/CurveEditorSelection.cpp b/Engine/Source/Editor/CurveEditor/Private/CurveEditorSelection.cpp
--- a/Engine/Source/Editor/CurveEditor/Private/CurveEditorSelection.cpp
+++ b/Engine/Source/Editor/CurveEditor/Private/CurveEditorSelection.cpp
@@ -29,7 +29,7 @@ void FCurveEditorSelection::Clear()
 	if (!bIsAlreadyEmpty)
 	{
 		CurveToSelectedKeys.Reset();
-		++SerialNumber;
+		PostChangesMade();
 	}
 }
 
@@ -48,6 +48,14 @@ void FCurveEditorSelection::AddInternal(
 	{
 		SelectedKeys.Add(Key, PointType);
 	}
+	
+	PostChangesMade();
+}
+
+void FCurveEditorSelection::PostChangesMade()
+{
+	++SerialNumber;
+	OnSelectionChangedBroadcaster.Broadcast();
 }
 
 const FKeyHandleSet* FCurveEditorSelection::FindForCurve(FCurveModelID InCurveID) const
@@ -77,6 +85,21 @@ bool FCurveEditorSelection::Contains(FCurveModelID CurveID, FKeyHandle KeyHandle
 	return SelectedKeys && SelectedKeys->Contains(KeyHandle, PointType);
 }
 
+void FCurveEditorSelection::IncrementOnSelectionChangedSuppressionCount()
+{
+	OnSelectionChangedBroadcaster.IncrementSuppressionCount();
+}
+
+void FCurveEditorSelection::DecrementOnSelectionChangedSuppressionCount()
+{
+	OnSelectionChangedBroadcaster.DecrementSuppressionCount();
+}
+
+FSimpleMulticastDelegate& FCurveEditorSelection::OnSelectionChanged()
+{
+	return OnSelectionChangedBroadcaster.OnChanged();
+}
+
 void FCurveEditorSelection::Add(FCurvePointHandle InHandle)
 {
 	Add(InHandle.CurveID, InHandle.PointType, InHandle.KeyHandle);
@@ -99,8 +122,6 @@ void FCurveEditorSelection::Add(FCurveModelID CurveID, ECurvePointType PointType
 			AddInternal(CurveID, PointType, Keys, bIncrementSerialNumber);
 		}
 	}
-
-	++SerialNumber;
 }
 
 void FCurveEditorSelection::Toggle(FCurvePointHandle InHandle)
@@ -127,9 +148,9 @@ void FCurveEditorSelection::Toggle(FCurveModelID CurveID, ECurvePointType PointT
 		{
 			CurveToSelectedKeys.Remove(CurveID);
 		}
+		
+		PostChangesMade();
 	}
-
-	++SerialNumber;
 }
 
 void FCurveEditorSelection::Remove(FCurvePointHandle InHandle)
@@ -151,15 +172,15 @@ void FCurveEditorSelection::Remove(FCurveModelID CurveID, ECurvePointType PointT
 		{
 			SelectedKeys.Remove(Key, PointType);
 		}
+		
+		PostChangesMade();
 	}
-
-	++SerialNumber;
 }
 
 void FCurveEditorSelection::Remove(FCurveModelID InCurveID)
 {
 	CurveToSelectedKeys.Remove(InCurveID);
-	++SerialNumber;
+	PostChangesMade();
 }
 
 void FKeyHandleSet::Add(FKeyHandle Handle, ECurvePointType PointType)
diff --git a/Engine/Source/Editor/CurveEditor/Private/CurveEditorSelectionPrivate.h b/Engine/Source/Editor/CurveEditor/Private/CurveEditorSelectionPrivate.h
--- a/Engine/Source/Editor/CurveEditor/Private/CurveEditorSelectionPrivate.h
+++ b/Engine/Source/Editor/CurveEditor/Private/CurveEditorSelectionPrivate.h
@@ -5,6 +5,8 @@
 #include "CurveEditorSelection.h"
 #include "Modification/Selection/SelectionChangeUtils.h"
 
+namespace UE::CurveEditor { class FScopedSelectionChangeEventSuppression; }
+
 namespace UE::CurveEditor::KeySelection 
 {
 /** Implements passkey pattern to restrict access to FCurveEditorSelection::AddInternal. */
@@ -33,4 +35,30 @@ class FSetSerialNumber
 		Selection.SerialNumber = NewSerialNumber;
 	}
 };
+
+/** Implements passkey pattern to restrict access to FCurveEditorSelection::IncrementOnSelectionChangedSuppressionCount. */
+class FIncrementOnSelectionChangedSuppressionCount
+{
+	friend FScopedSelectionChangeEventSuppression;
+
+	explicit FIncrementOnSelectionChangedSuppressionCount(
+		FCurveEditorSelection& Selection
+		)
+	{
+		Selection.IncrementOnSelectionChangedSuppressionCount();
+	}
+};
+
+/** Implements passkey pattern to restrict access to FCurveEditorSelection::DecrementOnSelectionChangedSuppressionCount. */
+class FDecrementOnSelectionChangedSuppressionCount
+{
+	friend FScopedSelectionChangeEventSuppression;
+
+	explicit FDecrementOnSelectionChangedSuppressionCount(
+		FCurveEditorSelection& Selection
+		)
+	{
+		Selection.DecrementOnSelectionChangedSuppressionCount();
+	}
+};
 }
diff --git a/Engine/Source/Editor/CurveEditor/Private/Modification/Selection/SelectionChangeCommand.cpp b/Engine/Source/Editor/CurveEditor/Private/Modification/Selection/SelectionChangeCommand.cpp
--- a/Engine/Source/Editor/CurveEditor/Private/Modification/Selection/SelectionChangeCommand.cpp
+++ b/Engine/Source/Editor/CurveEditor/Private/Modification/Selection/SelectionChangeCommand.cpp
@@ -4,6 +4,7 @@
 
 #include "CurveEditor.h"
 #include "Modification/Selection/SelectionChangeUtils.h"
+#include "Selection/ScopedSelectionChangeEventSuppression.h"
 
 namespace UE::CurveEditor
 {
@@ -11,6 +12,7 @@ void FSelectionChangeCommand::Apply(UObject* Object)
 {
 	if (const TSharedPtr<FCurveEditor> CurveEditor = GetCurveEditor())
 	{
+		const FScopedSelectionChangeEventSuppression SelectionChangedSuppressor(CurveEditor.ToSharedRef());
 		ApplySelectionChange(CurveEditor->Selection, DeltaChange);
 	}
 }
@@ -19,6 +21,7 @@ void FSelectionChangeCommand::Revert(UObject* Object)
 {
 	if (const TSharedPtr<FCurveEditor> CurveEditor = GetCurveEditor())
 	{
+		const FScopedSelectionChangeEventSuppression SelectionChangedSuppressor(CurveEditor.ToSharedRef());
 		RevertSelectionChange(CurveEditor->Selection, DeltaChange);
 	}
 }
diff --git a/Engine/Source/Editor/CurveEditor/Private/Modification/Utils/ScopedSelectionChange.cpp b/Engine/Source/Editor/CurveEditor/Private/Modification/Utils/ScopedSelectionChange.cpp
--- a/Engine/Source/Editor/CurveEditor/Private/Modification/Utils/ScopedSelectionChange.cpp
+++ b/Engine/Source/Editor/CurveEditor/Private/Modification/Utils/ScopedSelectionChange.cpp
@@ -18,10 +18,19 @@ FScopedSelectionChange::FScopedSelectionChange(TWeakPtr<FCurveEditor> InCurveEdi
 FScopedSelectionChange::FScopedSelectionChange(TWeakPtr<FCurveEditor> InCurveEditor, const FText& InDescription, EScopedSelectionChangeFlags InFlags)
 	: FScopedChangeBase(MoveTemp(InCurveEditor), InDescription)
 	, Flags(InFlags)
+	
 {
-	if (const TSharedPtr<FCurveEditor> CurveEditorPin = WeakCurveEditor.Pin())
+	const TSharedPtr<FCurveEditor> CurveEditorPin = WeakCurveEditor.Pin();
+	if (!CurveEditorPin)
+	{
+		return;
+	}
+
+	OriginalSelection = CurveEditorPin->Selection;
+	
+	if (EnumHasAnyFlags(Flags, EScopedSelectionChangeFlags::SuppressOnSelectionChangedEvent))
 	{
-		OriginalSelection = CurveEditorPin->Selection;
+		SelectionChangedSuppressor.Emplace(CurveEditorPin.ToSharedRef());
 	}
 }
 
@@ -49,7 +58,14 @@ FScopedSelectionChange::~FScopedSelectionChange()
 	
 	if (!bHasBeenCancelled)
 	{
+		// Start a transaction...
+		const FScopedTransaction Transaction(Description);
+		
 		TryAppendCommand(MakeUnique<FSelectionChangeCommand>(CurveEditorPin, MoveTemp(DeltaChange)));
+		
+		// ... so that if the suppression counter is decremented to 0, the event fires as part of the transaction.
+		// This way any changes made will be recorded as part of the same transaction.
+		SelectionChangedSuppressor.Reset(); 
 	}
 	else
 	{
diff --git a/Engine/Source/Editor/CurveEditor/Private/Selection/ScopedSelectionChangeEventSuppression.cpp b/Engine/Source/Editor/CurveEditor/Private/Selection/ScopedSelectionChangeEventSuppression.cpp
new file mode 100644
--- /dev/null
+++ b/Engine/Source/Editor/CurveEditor/Private/Selection/ScopedSelectionChangeEventSuppression.cpp
@@ -0,0 +1,23 @@
+﻿// Copyright Epic Games, Inc. All Rights Reserved.
+
+#include "Selection/ScopedSelectionChangeEventSuppression.h"
+
+#include "CurveEditor.h"
+#include "CurveEditorSelectionPrivate.h"
+
+namespace UE::CurveEditor
+{
+FScopedSelectionChangeEventSuppression::FScopedSelectionChangeEventSuppression(const TSharedRef<FCurveEditor>& InCurveEditor)
+	: WeakCurveEditor(InCurveEditor)
+{
+	KeySelection::FIncrementOnSelectionChangedSuppressionCount(InCurveEditor->Selection);
+}
+
+FScopedSelectionChangeEventSuppression::~FScopedSelectionChangeEventSuppression()
+{
+	if (const TSharedPtr<FCurveEditor> CurveEditorPin = WeakCurveEditor.Pin())
+	{
+		KeySelection::FDecrementOnSelectionChangedSuppressionCount(CurveEditorPin->Selection);
+	}
+}
+}
diff --git a/Engine/Source/Editor/CurveEditor/Public/CurveEditorSelection.h b/Engine/Source/Editor/CurveEditor/Public/CurveEditorSelection.h
--- a/Engine/Source/Editor/CurveEditor/Public/CurveEditorSelection.h
+++ b/Engine/Source/Editor/CurveEditor/Public/CurveEditorSelection.h
@@ -12,6 +12,7 @@
 #include "Curves/KeyHandle.h"
 #include "HAL/PlatformCrt.h"
 #include "Misc/Optional.h"
+#include "Misc/SuppressibleEventBroadcaster.h"
 #include "Templates/SharedPointer.h"
 
 #define UE_API CURVEEDITOR_API
@@ -19,6 +20,8 @@
 class FCurveEditor;
 namespace UE::CurveEditor::KeySelection { class FAddInternal; }
 namespace UE::CurveEditor::KeySelection { class FSetSerialNumber; }
+namespace UE::CurveEditor::KeySelection { class FIncrementOnSelectionChangedSuppressionCount; }
+namespace UE::CurveEditor::KeySelection { class FDecrementOnSelectionChangedSuppressionCount; }
 
 struct FCurvePointHandle;
 
@@ -91,6 +94,8 @@ struct FCurveEditorSelection
 {
 	friend UE::CurveEditor::KeySelection::FAddInternal;
 	friend UE::CurveEditor::KeySelection::FSetSerialNumber;
+	friend UE::CurveEditor::KeySelection::FIncrementOnSelectionChangedSuppressionCount;
+	friend UE::CurveEditor::KeySelection::FDecrementOnSelectionChangedSuppressionCount;
 	
 	/**
 	 * Default constructor
@@ -138,6 +143,9 @@ struct FCurveEditorSelection
 	 */
 	UE_API bool Contains(FCurveModelID CurveID, FKeyHandle KeyHandle, ECurvePointType PointType) const;
 
+	/** Invoked when the selection changes. */
+	UE_API FSimpleMulticastDelegate& OnSelectionChanged();
+
 public:
 
 	/**
@@ -213,9 +221,20 @@ private:
 
 	/** A map of selected handles stored by curve ID */
 	TMap<FCurveModelID, FKeyHandleSet> CurveToSelectedKeys;
+
+	/** Invoked when the selection changes. */
+	UE::CurveEditor::FSuppressibleEventBroadcaster OnSelectionChangedBroadcaster;
+
+	/** Prevents OnSelectionChanged from being Broadcast until a matching call to DecrementOnSelectionChangedSuppressionCount is made. */
+	void IncrementOnSelectionChangedSuppressionCount();
+	/** Allows OnSelectionChanged to be Broadcast if the counter reaches 0. Invokes OnSelectionChanged if a change has been made since. */
+	void DecrementOnSelectionChangedSuppressionCount();
 	
 	/** Adds the specified key handles from the selection without validating whether CurveId and Keys exists in the owning curve editor. */
-	UE_API void AddInternal(FCurveModelID CurveID, ECurvePointType PointType, TArrayView<const FKeyHandle> Keys, bool bIncrementSerialNumber = true);
+	void AddInternal(FCurveModelID CurveID, ECurvePointType PointType, TArrayView<const FKeyHandle> Keys, bool bIncrementSerialNumber = true);
+
+	/** Shared logic after some changes have been made. */
+	void PostChangesMade();
 };
 
 #undef UE_API
diff --git a/Engine/Source/Editor/CurveEditor/Public/Misc/SuppressibleEventBroadcaster.h b/Engine/Source/Editor/CurveEditor/Public/Misc/SuppressibleEventBroadcaster.h
new file mode 100644
--- /dev/null
+++ b/Engine/Source/Editor/CurveEditor/Public/Misc/SuppressibleEventBroadcaster.h
@@ -0,0 +1,74 @@
+﻿// Copyright Epic Games, Inc. All Rights Reserved.
+
+#pragma once
+
+#include "Delegates/Delegate.h"
+
+namespace UE::CurveEditor
+{
+/**
+ * Wraps a simple multicast delegate with a counter.
+ *
+ * While the counter is positive, the event cannot be broadcast.
+ * Once the counter reaches 0, the event is broadcast.
+ *
+ * This is a reusable optimization that is useful to change-type events.
+ */
+class FSuppressibleEventBroadcaster
+{
+public:
+
+	void IncrementSuppressionCount()
+	{
+		++SuppressionCount;
+	}
+
+	void DecrementSuppressionCount()
+	{
+		if (!ensure(SuppressionCount > 0))
+		{
+			return;
+		}
+		
+		--SuppressionCount;
+
+		if (SuppressionCount == 0 && bWasEventSuppressed)
+		{
+			bWasEventSuppressed = false;
+			Broadcast();
+		}
+	}
+
+	void Broadcast()
+	{
+		if (SuppressionCount == 0)
+		{
+			bWasEventSuppressed = false;
+			Event.Broadcast();
+		}
+		else
+		{
+			bWasEventSuppressed = true;
+		}
+	}
+
+	/** @return The event that is broadcast. This exists only for convenience: by design, you're NOT supposed to call Broadcast on this directly. */
+	FSimpleMulticastDelegate& OnChanged() { return Event; }
+	
+private:
+
+	/** When 0, the event can be broadcast. */
+	int32 SuppressionCount = 0;
+
+	/**
+	 * This is set to true when Broadcast is called while SuppressionCount > 0.
+	 * If true, then DecrementSuppressionCount will broadcast the event when count reaches 0.
+	 */
+	bool bWasEventSuppressed = false;
+
+	/** The event to broadcast */
+	FSimpleMulticastDelegate Event;
+};
+}
+
+
diff --git a/Engine/Source/Editor/CurveEditor/Public/Modification/Utils/ScopedSelectionChange.h b/Engine/Source/Editor/CurveEditor/Public/Modification/Utils/ScopedSelectionChange.h
--- a/Engine/Source/Editor/CurveEditor/Public/Modification/Utils/ScopedSelectionChange.h
+++ b/Engine/Source/Editor/CurveEditor/Public/Modification/Utils/ScopedSelectionChange.h
@@ -4,7 +4,9 @@
 
 #include "CurveEditorSelection.h"
 #include "Misc/EnumClassFlags.h"
+#include "Misc/Optional.h"
 #include "ScopedChangeBase.h"
+#include "Selection/ScopedSelectionChangeEventSuppression.h"
 
 #define UE_API CURVEEDITOR_API
 
@@ -17,13 +19,24 @@ enum class EScopedSelectionChangeFlags : uint8
 	/** Revert the change if FScopedSelectionChange or surrounding FScopedTransaction becomes cancelled. */
 	RevertOnCancel = 1 << 0,
 
-	Default = RevertOnCancel
+	/**
+	 * While the scope is constructed, also construct a FScopedSelectionChangeEventSuppression.
+	 * 
+	 * If set, then the curve editor's FCurveEditorSelection::OnSelectionChanged will not be broadcast during the scope. Once the scope ends, and a
+	 * change has been made to selection, FCurveEditorSelection::OnSelectionChanged is invoked.
+	 */
+	SuppressOnSelectionChangedEvent = 1 << 1,
+
+	Default = RevertOnCancel | SuppressOnSelectionChangedEvent
 };
 ENUM_CLASS_FLAGS(EScopedSelectionChangeFlags);
 	
 /**
- * Saves the current selection state and on destruction compares it against the then current selection. If the selection has changed, it appends
- * an undo-able action to the curve editor transaction manager.
+ * Saves the current selection state and on destruction compares it against the then current selection.
+ * If the selection has changed when the scope ends, then
+ * - it appends an undo-able action to the curve editor transaction manager,
+ * - conditionally reverts the change if you called FScopedTransaction::Cancel on the surrounding transaction,
+ * - conditionally broadcasts FCurveEditorSelection::OnSelectionChanged if you set EScopedSelectionChangeFlags::SuppressOnSelectionChangedEvent.
  *
  * You can safely use this within the scope of a FScopedTransaction. In that case, the command will be added as sub-transaction to the parent
  * transaction.
@@ -51,6 +64,9 @@ private:
 
 	/** Flags that affect the behaviour of this change. */
 	const EScopedSelectionChangeFlags Flags;
+
+	/** Supresses FCurveEditorSelection::OnSelectionChanged until the scope ends. Only set if the flags requested it. */
+	TOptional<FScopedSelectionChangeEventSuppression> SelectionChangedSuppressor;
 	
 	/** The selection the editor had when the transaction was started.  */
 	FCurveEditorSelection OriginalSelection;
diff --git a/Engine/Source/Editor/CurveEditor/Public/Selection/ScopedSelectionChangeEventSuppression.h b/Engine/Source/Editor/CurveEditor/Public/Selection/ScopedSelectionChangeEventSuppression.h
new file mode 100644
--- /dev/null
+++ b/Engine/Source/Editor/CurveEditor/Public/Selection/ScopedSelectionChangeEventSuppression.h
@@ -0,0 +1,47 @@
+﻿// Copyright Epic Games, Inc. All Rights Reserved.
+
+#pragma once
+
+#include "Templates/SharedPointer.h"
+#include "Templates/UnrealTemplate.h"
+
+class FCurveEditor;
+
+namespace UE::CurveEditor
+{
+/**
+ * Prevents FCurveEditorSelection::OnSelectionChanged from being broadcast while the constructed scope is open.
+ * Once all scopes end, and if changes to selection have been made, FCurveEditorSelection::OnSelectionChanged is broadcast exactly once.
+ * Nesting FScopedSelectionChangeEventSuppression is supported.
+ * 
+ * This is an optimization if you are going to make multiple changes to selection and want the delegate to be broadcast after you're done.
+ *
+ * Example usage:
+ * TSharedRef<FCurveEditor> MyCurveEditor = ...;
+ * {
+ *		FScopedSelectionChangeEventSuppression SuppressChangeEvents(MyCurveEditor);
+ *		MyCurveEditor->Selection.Toggle(FooCurveId, ECurvePointType::Key, FooKeyHandle); // Normally this would invoke OnSelectionChanged. 
+ *		MyCurveEditor->Selection.Toggle(BarCurveId, ECurvePointType::Key, BarKeyHandle); // So would this
+ *		{
+ *			FScopedSelectionChangeEventSuppression AnotherChange(MyCurveEditor);
+ *			MyCurveEditor->Selection.Toggle(BazCurveId, ECurvePointType::Key, BazKeyHandle);
+ *			//~ FCurveEditorSelection::OnSelectionChanged not invoked due to the outer scope.
+ *		}
+ *		
+ *		// ~FScopedSelectionChangeEventSuppression will now invoke FCurveEditorSelection::OnSelectionChanged once exactly.
+ * }
+ */
+class CURVEEDITOR_API FScopedSelectionChangeEventSuppression : public FNoncopyable
+{
+public:
+
+	explicit FScopedSelectionChangeEventSuppression(const TSharedRef<FCurveEditor>& InCurveEditor);
+	~FScopedSelectionChangeEventSuppression();
+
+private:
+
+	/** Handles the case of curve editor being destroyed during the scope, which should not happen but we handle for safety. */
+	const TWeakPtr<FCurveEditor> WeakCurveEditor;
+};
+}
+
