Curve Editor: Fix undo / redo when keys are replaced. The problem was that we skipped diffing certain types of changes in some of the operations.

diff --git a/Engine/Source/Editor/CurveEditor/Private/CurveEditor.cpp b/Engine/Source/Editor/CurveEditor/Private/CurveEditor.cpp
--- a/Engine/Source/Editor/CurveEditor/Private/CurveEditor.cpp
+++ b/Engine/Source/Editor/CurveEditor/Private/CurveEditor.cpp
@@ -1864,7 +1864,10 @@ void FCurveEditor::PasteKeys(UE::CurveEditor::FKeyPasteArgs InArgs)
 	
 	const FScopedTransaction Transaction(LOCTEXT("PasteKeys", "Paste Keys")); // Selection and key change should be part of same undo / redo op.
 	const FScopedSelectionChange SelectionChange(SharedThis(this));
-	constexpr ECurveChangeFlags Flags = ECurveChangeFlags::AddKeys | ECurveChangeFlags::RemoveKeys | ECurveChangeFlags::KeyAttributes;
+	constexpr ECurveChangeFlags Flags = ECurveChangeFlags::AddKeys | ECurveChangeFlags::RemoveKeys | ECurveChangeFlags::KeyAttributes
+		// Even though paste uses FCurveModel::AddKeys - if you call AddKeys with an x-position that is over an existing key, then no new key is added.
+		// The pre-existing key is moved. So we ECurveChangeFlags::MoveKeys, as well.
+		| ECurveChangeFlags::MoveKeys;
 	const FScopedCurveChange KeyChange(FCurvesSnapshotBuilder(SharedThis(this), TargetCurves, Flags));
 	Selection.Clear();
 
diff --git a/Engine/Source/Editor/CurveEditor/Private/DragOperations/CurveEditorDragOperation_MoveKeys.cpp b/Engine/Source/Editor/CurveEditor/Private/DragOperations/CurveEditorDragOperation_MoveKeys.cpp
--- a/Engine/Source/Editor/CurveEditor/Private/DragOperations/CurveEditorDragOperation_MoveKeys.cpp
+++ b/Engine/Source/Editor/CurveEditor/Private/DragOperations/CurveEditorDragOperation_MoveKeys.cpp
@@ -35,7 +35,12 @@ void FCurveEditorDragOperation_MoveKeys::OnBeginDrag(FVector2D InitialPosition,
 {
 	using namespace UE::CurveEditor;
 	
-	ScopedKeyChange.Emplace(FCurvesSnapshotBuilder(CurveEditor->AsShared(), ECurveChangeFlags::MoveKeys).TrackSelectedKeys());
+	ScopedKeyChange.Emplace(
+		// If you move a key over the same time as another key, that key is replaced so we need ECurveChangeFlags::RemoveKeys.
+		// This is easy to repro if you enable key snapping.
+		FCurvesSnapshotBuilder(CurveEditor->AsShared(), ECurveChangeFlags::MoveKeys | ECurveChangeFlags::RemoveKeys)
+		.TrackSelectedCurves()
+		);
 
 	KeysByCurve.Reset();
 	CurveEditor->SuppressBoundTransformUpdates(true);
diff --git a/Engine/Source/Editor/CurveEditor/Private/Modification/Keys/GenericCurveChangeUtils.cpp b/Engine/Source/Editor/CurveEditor/Private/Modification/Keys/GenericCurveChangeUtils.cpp
--- a/Engine/Source/Editor/CurveEditor/Private/Modification/Keys/GenericCurveChangeUtils.cpp
+++ b/Engine/Source/Editor/CurveEditor/Private/Modification/Keys/GenericCurveChangeUtils.cpp
@@ -17,9 +17,19 @@ void ApplyChange(
 	)
 {
 	SCOPED_CURVE_EDITOR_TRACE(ApplyGenericChange);
+
+	// Order matters. Must move keys before adding them in KeyInsertion::ApplyChange.
+	// If we added keys first and a key we're about to move is already at the position, we'd end up not creating any key.
+	// This is because we're not allowed to have keys stacked on the same x-position.
+	// Example: First, move key 909 from x=0 to x=1. Then, add key 891 from to position x=0.
+	
+	// Note: What about e.g. moving key 230 from x=0 to x=1 and adding key 450 to position x=1? We needn't worry about this because it's an invalid
+	// operation to begin with. All FGenericCurveChangeData are constructed from changes that were performed.
+	// This hypothetical operation shouldn't have occured in the first place.
+	MoveKeys::ApplyChange(InCurves, InDeltaChange.MoveKeysData);
+	
 	KeyRemoval::ApplyChange(InCurves, InDeltaChange.RemoveKeysData, InCurrentSliderPos);
 	KeyInsertion::ApplyChange(InCurves, InDeltaChange.AddKeysData);
-	MoveKeys::ApplyChange(InCurves, InDeltaChange.MoveKeysData);
 	KeyAttributes::ApplyChange(InCurves, InDeltaChange.KeyAttributeData);
 	CurveAttributes::ApplyChange(InCurves, InDeltaChange.CurveAttributeData);
 }
@@ -29,9 +39,21 @@ void RevertChange(
 	)
 {
 	SCOPED_CURVE_EDITOR_TRACE(RevertGenericChange);
+	
+	// Order matters. Must move keys before adding them back in KeyRemoval::ApplyChange.
+	// If we moved a key on top of another key's time, that replaces the key.
+	// If we tried adding a key back before moving the key that currently happens to be there (but is moved by our up), we end up with a single key
+	// instead of with two.
+	// This is because we're not allowed to have keys stacked on the same x-position.
+	// Example: First, move key 909 from x=1 to x=0. Then, remove key 891 from position x=0.
+	
+	// Note: What about e.g. moving key 230 from x=0 to x=1 and adding key 450 to position x=1? We needn't worry about this because it's an invalid
+	// operation to begin with. All FGenericCurveChangeData are constructed from changes that were performed.
+	// This hypothetical operation shouldn't have occured in the first place.
+	MoveKeys::RevertChange(InCurves, InDeltaChange.MoveKeysData);
+	
 	KeyInsertion::RevertChange(InCurves, InDeltaChange.AddKeysData, InCurrentSliderPos);
 	KeyRemoval::RevertChange(InCurves, InDeltaChange.RemoveKeysData);
-	MoveKeys::RevertChange(InCurves, InDeltaChange.MoveKeysData);
 	KeyAttributes::RevertChange(InCurves, InDeltaChange.KeyAttributeData);
 	CurveAttributes::RevertChange(InCurves, InDeltaChange.CurveAttributeData);
 }
