Curve Editor: If your curves are linear, adding a key is now also be linear.

diff --git a/Engine/Source/Editor/CurveEditor/Private/Views/SInteractiveCurveEditorView.cpp b/Engine/Source/Editor/CurveEditor/Private/Views/SInteractiveCurveEditorView.cpp
--- a/Engine/Source/Editor/CurveEditor/Private/Views/SInteractiveCurveEditorView.cpp
+++ b/Engine/Source/Editor/CurveEditor/Private/Views/SInteractiveCurveEditorView.cpp
@@ -2044,17 +2044,16 @@ FKeyAttributes SInteractiveCurveEditorView::GetDefaultKeyAttributesForCurveTime(
 {
 	FKeyAttributes KeyAttributes = CurveEditor.GetDefaultKeyAttribute().Get();
 
-	TPair<ERichCurveInterpMode, ERichCurveTangentMode> Modes = CurveModel.GetInterpolationMode(EvalTime, ERichCurveInterpMode::RCIM_Linear, ERichCurveTangentMode::RCTM_Auto);
-	if (Modes.Key != ERichCurveInterpMode::RCIM_Linear)
-	{
-		KeyAttributes.SetInterpMode(Modes.Key);
-	}
-	if (Modes.Value != ERichCurveTangentMode::RCTM_Auto)
-	{
-		KeyAttributes.SetTangentMode(Modes.Value);
-	}
-
-	return KeyAttributes;
+	// Give the model a chance to override the default interpolation modes
+    const TPair<ERichCurveInterpMode, ERichCurveTangentMode> Modes = CurveModel.GetInterpolationMode(
+    	EvalTime,
+    	KeyAttributes.HasInterpMode() ? KeyAttributes.GetInterpMode() : RCIM_Linear,
+    	KeyAttributes.HasTangentMode() ? KeyAttributes.GetTangentMode() : RCTM_Auto
+    	);
+    KeyAttributes.SetInterpMode(Modes.Key);
+    KeyAttributes.SetTangentMode(Modes.Value);
+	
+    return KeyAttributes;
 }
 
 void SInteractiveCurveEditorView::AddKeyAtScrubTime(TSet<FCurveModelID> ForCurves)
