Curve Editor: Fix retime, lattice, and multi-scale tools not undoing correctly when keys are replaced.

diff --git a/Engine/Plugins/Editor/CurveEditorTools/Source/CurveEditorTools/Private/Tools/CurveEditorMultiScaleTool.cpp b/Engine/Plugins/Editor/CurveEditorTools/Source/CurveEditorTools/Private/Tools/CurveEditorMultiScaleTool.cpp
--- a/Engine/Plugins/Editor/CurveEditorTools/Source/CurveEditorTools/Private/Tools/CurveEditorMultiScaleTool.cpp
+++ b/Engine/Plugins/Editor/CurveEditorTools/Source/CurveEditorTools/Private/Tools/CurveEditorMultiScaleTool.cpp
@@ -388,7 +388,9 @@ void FCurveEditorMultiScaleTool::OnDragStart()
 
 	using namespace UE::CurveEditor;
 	ActiveTransaction = MakeUnique<FScopedCurveChange>(
-		FCurvesSnapshotBuilder(CurveEditor, ECurveChangeFlags::KeyAttributes | ECurveChangeFlags::MoveKeys).TrackSelectedCurves(),
+		FCurvesSnapshotBuilder(
+			CurveEditor, ECurveChangeFlags::KeyAttributes | ECurveChangeFlags::MoveKeysAndRemoveStackedKeys
+			).TrackSelectedCurves(),
 		LOCTEXT("CurveEditorMultiScaleToolTransaction", "Multi Scale Key(s)")
 		);
 	TangentFlattener.Emplace(*CurveEditor);
diff --git a/Engine/Plugins/Editor/CurveEditorTools/Source/CurveEditorTools/Private/Tools/CurveEditorTransformTool.cpp b/Engine/Plugins/Editor/CurveEditorTools/Source/CurveEditorTools/Private/Tools/CurveEditorTransformTool.cpp
--- a/Engine/Plugins/Editor/CurveEditorTools/Source/CurveEditorTools/Private/Tools/CurveEditorTransformTool.cpp
+++ b/Engine/Plugins/Editor/CurveEditorTools/Source/CurveEditorTools/Private/Tools/CurveEditorTransformTool.cpp
@@ -788,7 +788,9 @@ void FCurveEditorTransformTool::OnDragStart()
 
 	using namespace UE::CurveEditor;
 	ActiveTransaction = MakeUnique<FScopedCurveChange>(
-		FCurvesSnapshotBuilder(CurveEditor, ECurveChangeFlags::KeyAttributes | ECurveChangeFlags::MoveKeys).TrackSelectedCurves(),
+		FCurvesSnapshotBuilder(CurveEditor,
+			ECurveChangeFlags::KeyAttributes | ECurveChangeFlags::MoveKeysAndRemoveStackedKeys
+			).TrackSelectedCurves(),
 		LOCTEXT("CurveEditorTransformToolTransaction", "Transform Key(s)")
 		);
 	TangentFlattener.Emplace(*CurveEditor);
diff --git a/Engine/Plugins/Editor/CurveEditorTools/Source/CurveEditorTools/Private/Tools/Lattice/CurveEditorLatticeTool.h b/Engine/Plugins/Editor/CurveEditorTools/Source/CurveEditorTools/Private/Tools/Lattice/CurveEditorLatticeTool.h
--- a/Engine/Plugins/Editor/CurveEditorTools/Source/CurveEditorTools/Private/Tools/Lattice/CurveEditorLatticeTool.h
+++ b/Engine/Plugins/Editor/CurveEditorTools/Source/CurveEditorTools/Private/Tools/Lattice/CurveEditorLatticeTool.h
@@ -158,7 +158,10 @@ private:
 			{
 				using namespace UE::CurveEditor;
 				return MakeUnique<FScopedCurveChange>(
-					FCurvesSnapshotBuilder(InCurveEditor, ECurveChangeFlags::MoveKeys | ECurveChangeFlags::KeyAttributes).TrackSelectedCurves()
+					FCurvesSnapshotBuilder(
+						InCurveEditor,
+						ECurveChangeFlags::MoveKeysAndRemoveStackedKeys | ECurveChangeFlags::KeyAttributes
+						).TrackSelectedCurves()
 					);
 			}())
 		{}
diff --git a/Engine/Plugins/Editor/CurveEditorTools/Source/CurveEditorTools/Private/Tools/Retime/CurveEditorRetimeTool.h b/Engine/Plugins/Editor/CurveEditorTools/Source/CurveEditorTools/Private/Tools/Retime/CurveEditorRetimeTool.h
--- a/Engine/Plugins/Editor/CurveEditorTools/Source/CurveEditorTools/Private/Tools/Retime/CurveEditorRetimeTool.h
+++ b/Engine/Plugins/Editor/CurveEditorTools/Source/CurveEditorTools/Private/Tools/Retime/CurveEditorRetimeTool.h
@@ -8,6 +8,7 @@
 #include "ScopedTransaction.h"
 #include "RetimeAnchor.h"
 #include "RetimeDragData.h"
+#include "Modification/Utils/ScopedSelectionChange.h"
 #include "CurveEditorRetimeTool.generated.h"
 
 class FCurveEditor;
@@ -80,14 +81,16 @@ private:
 	struct FTransaction
 	{
 		FScopedTransaction Transaction;
-		UE::CurveEditor::FScopedCurveChange KeyChange;
+		CurveEditor::FScopedCurveChange KeyChange;
 
 		explicit FTransaction(TWeakPtr<FCurveEditor> InCurveEditor, const FText& Description)
 			: Transaction(Description)
 			, KeyChange([InCurveEditor]
 			{
 				using namespace UE::CurveEditor;
-				return FCurvesSnapshotBuilder(InCurveEditor, ECurveChangeFlags::KeyAttributes | ECurveChangeFlags::MoveKeys).TrackAllCurves();
+				return FCurvesSnapshotBuilder(
+					InCurveEditor, ECurveChangeFlags::KeyAttributes | ECurveChangeFlags::MoveKeysAndRemoveStackedKeys
+					).TrackAllCurves();
 			}())
 		{}
 	};
diff --git a/Engine/Source/Editor/CurveEditor/Private/CurveEditor.cpp b/Engine/Source/Editor/CurveEditor/Private/CurveEditor.cpp
--- a/Engine/Source/Editor/CurveEditor/Private/CurveEditor.cpp
+++ b/Engine/Source/Editor/CurveEditor/Private/CurveEditor.cpp
@@ -1864,10 +1864,7 @@ void FCurveEditor::PasteKeys(UE::CurveEditor::FKeyPasteArgs InArgs)
 	
 	const FScopedTransaction Transaction(LOCTEXT("PasteKeys", "Paste Keys")); // Selection and key change should be part of same undo / redo op.
 	const FScopedSelectionChange SelectionChange(SharedThis(this));
-	constexpr ECurveChangeFlags Flags = ECurveChangeFlags::AddKeys | ECurveChangeFlags::RemoveKeys | ECurveChangeFlags::KeyAttributes
-		// Even though paste uses FCurveModel::AddKeys - if you call AddKeys with an x-position that is over an existing key, then no new key is added.
-		// The pre-existing key is moved. So we ECurveChangeFlags::MoveKeys, as well.
-		| ECurveChangeFlags::MoveKeys;
+	constexpr ECurveChangeFlags Flags = ECurveChangeFlags::MoveKeysAndRemoveStackedKeys | ECurveChangeFlags::RemoveKeys | ECurveChangeFlags::KeyAttributes;
 	const FScopedCurveChange KeyChange(FCurvesSnapshotBuilder(SharedThis(this), TargetCurves, Flags));
 	Selection.Clear();
 
@@ -2167,7 +2164,7 @@ void FCurveEditor::SmartSnapSelection()
 	const FScopedTransaction Transaction(LOCTEXT("SmartSnapKeys", "Smart Snap")); // Selection and key change should be part of same undo / redo op.
 	const FScopedSelectionChange SelectionChange(SharedThis(this));
 	const FScopedCurveChange KeyChange(
-		FCurvesSnapshotBuilder(SharedThis(this), ECurveChangeFlags::MoveKeys | ECurveChangeFlags::RemoveKeys).TrackSelectedCurves()
+		FCurvesSnapshotBuilder(SharedThis(this), ECurveChangeFlags::MoveKeysAndRemoveStackedKeys).TrackSelectedCurves()
 		);
 	
 	TMap<FCurveModelID, FKeyHandleSet> OutKeysToSelect;
diff --git a/Engine/Source/Editor/CurveEditor/Private/DragOperations/CurveEditorDragOperation_MoveKeys.cpp b/Engine/Source/Editor/CurveEditor/Private/DragOperations/CurveEditorDragOperation_MoveKeys.cpp
--- a/Engine/Source/Editor/CurveEditor/Private/DragOperations/CurveEditorDragOperation_MoveKeys.cpp
+++ b/Engine/Source/Editor/CurveEditor/Private/DragOperations/CurveEditorDragOperation_MoveKeys.cpp
@@ -36,9 +36,7 @@ void FCurveEditorDragOperation_MoveKeys::OnBeginDrag(FVector2D InitialPosition,
 	using namespace UE::CurveEditor;
 	
 	ScopedKeyChange.Emplace(
-		// If you move a key over the same time as another key, that key is replaced so we need ECurveChangeFlags::RemoveKeys.
-		// This is easy to repro if you enable key snapping.
-		FCurvesSnapshotBuilder(CurveEditor->AsShared(), ECurveChangeFlags::MoveKeys | ECurveChangeFlags::RemoveKeys)
+		FCurvesSnapshotBuilder(CurveEditor->AsShared(), ECurveChangeFlags::MoveKeysAndRemoveStackedKeys)
 		.TrackSelectedCurves()
 		);
 
diff --git a/Engine/Source/Editor/CurveEditor/Public/Modification/Keys/ECurveChangeFlags.h b/Engine/Source/Editor/CurveEditor/Public/Modification/Keys/ECurveChangeFlags.h
--- a/Engine/Source/Editor/CurveEditor/Public/Modification/Keys/ECurveChangeFlags.h
+++ b/Engine/Source/Editor/CurveEditor/Public/Modification/Keys/ECurveChangeFlags.h
@@ -26,7 +26,15 @@ enum class ECurveChangeFlags : uint8
 	/** Detect changes to the curve attributes, e.g. pre and post extrapolation, etc. */
 	CurveAttributes = 1 << 4,
 
+	/**
+	 * Use when you move keys and the keys may end up being stacked (i.e. have the same x positions).
+	 * When keys have the same x positions, curve editor removes the stacked keys so there's only 1 at the same x position.
+	 * So removed keys must be tracked.
+	 */
+	MoveKeysAndRemoveStackedKeys = MoveKeys | RemoveKeys,
+	/** Track anything that modifies key data. */
 	KeyData = MoveKeys | AddKeys | RemoveKeys | KeyAttributes,
+	/** All data should be tracked.*/
 	All = KeyData | CurveAttributes
 };
 ENUM_CLASS_FLAGS(ECurveChangeFlags);
