Curve Editor: Fix FScopedCurveChange not grouping per-curve mini-transactions. There was no issue yet, this is just to future proof any code that may use FScopedCurveChange but not start a FScopedTransaction of their own.

diff --git a/Engine/Source/Editor/CurveEditor/Private/Modification/Utils/ScopedCurveChange.cpp b/Engine/Source/Editor/CurveEditor/Private/Modification/Utils/ScopedCurveChange.cpp
--- a/Engine/Source/Editor/CurveEditor/Private/Modification/Utils/ScopedCurveChange.cpp
+++ b/Engine/Source/Editor/CurveEditor/Private/Modification/Utils/ScopedCurveChange.cpp
@@ -104,6 +104,11 @@ FScopedCurveChange::~FScopedCurveChange()
 	
 	if (!bHasBeenCancelled)
 	{
+		// The ProcessDiffs will call FTransactionManager::AppendChange repeatedly - once for each curve changed.
+		// As each call starts a FScopedTransaction, we should group them all under one transaction in case whatever code uses FScopedCurveChange
+		// has not started an outer FScopedTransaction of their own.
+		const FScopedTransaction GroupTransactions(Description);
+		
 		FMultiGenericCommandAppender CommandAppender(CurveEditorPin);
 		ChangeDiff.ProcessDiffs(CommandAppender);
 		LogPerfData(CommandAppender.GetAllocationSize()); 
