Curve Editor: Cleanse selection of stale keys, i.e. FKeyHandles and FCurveModelIds present in the selection but that are not referenced by the FCurveEditor.

Keys are cleansed when:
- Undo / redo
- Appending FGenericCurveChangeCommand

This fixes transform, and multi-scale tool bounds expanding to {0, 0} after stacking keys. The problem was that calling SetKeyPositions removed stacked keys and selection still referenced those duplicate keys. As such GetKeyPositions would return { 0, 0 }, which would expand the bounds.

diff --git a/Engine/Source/Editor/CurveEditor/Private/CurveEditor.cpp b/Engine/Source/Editor/CurveEditor/Private/CurveEditor.cpp
--- a/Engine/Source/Editor/CurveEditor/Private/CurveEditor.cpp
+++ b/Engine/Source/Editor/CurveEditor/Private/CurveEditor.cpp
@@ -63,6 +63,8 @@
 #include "Misc/SmartSnap.h"
 #include "Modification/Utils/ScopedCurveChange.h"
 #include "Modification/Utils/ScopedSelectionChange.h"
+#include "Selection/SelectionCleanser.h"
+#include "Selection/SelectionUtils.h"
 #include "Widgets/Colors/SColorPicker.h"
 #include "Widgets/Notifications/SNotificationList.h"
 
@@ -142,6 +144,8 @@ void FCurveEditor::InitCurveEditor(const FCurveEditorInitParams& InInitParams)
 	// Listen to global undo so we can fix up our selection state for keys that no longer exist.
 	GEditor->RegisterForUndo(this);
 	TransactionManager = MakeUnique<UE::CurveEditor::FTransactionManager>(SharedThis(this));
+
+	SelectionCleanser = MakePimpl<UE::CurveEditor::FSelectionCleanser>(SharedThis(this));
 }
 
 int32 FCurveEditor::GetSupportedTangentTypes()
@@ -2617,57 +2621,12 @@ bool FCurveEditor::IsActiveBufferedCurve(const TUniquePtr<IBufferedCurveModel>&
 	return false;
 }
 
-static TAutoConsoleVariable<bool> CVarDisableKeyCleansing(
-	TEXT("CurveEditor.EnableCurveCleansing"), false, TEXT("After undo operations, whether to remove invalid keys from the key selection.")
-	);
-
 void FCurveEditor::PostUndo(bool bSuccess)
 {
 	if (WeakPanel.IsValid())
 	{
 		WeakPanel.Pin()->PostUndo();
 	}
-
-	// Temporary hack while we're testing the new undo / redo system for key selection.
-	// Disabled by default - if animators notice problems, they can enable this CVar. We don't expect issues with this off by default.
-	// If no issues are found in 2 weeks, we'll just remove the entirety of the below code.
-	// 19th of Feb 2025. 
-	if (!CVarDisableKeyCleansing.GetValueOnGameThread())
-	{
-		return;
-	}
-
-	// If you create keys and then undo them the selection set still thinks there's keys selected.
-	// This presents issues with context menus and other things that are activated when there is a selection set.
-	// To fix this, we have to loop through all of our curve models, and re-select only the key handles that were
-	// previously selected that still exist. Ugly, but reasonably functional.
-	TMap<FCurveModelID, FKeyHandleSet> SelectionSet = Selection.GetAll();
-	for (const TPair<FCurveModelID, FKeyHandleSet>& Set : SelectionSet)
-	{
-		FCurveModel* CurveModel = FindCurve(Set.Key);
-
-		// If the entire curve was removed, just dump that out of the selection set.
-		if (!CurveModel)
-		{
-			Selection.Remove(Set.Key);
-			continue;
-		}
-		// Get all of the key handles from this curve.
-		const TArray<FKeyHandle> KeyHandles = CurveModel->GetAllKeys();
-
-		// The set handles will be mutated as we remove things so we need a copy that we can iterate through.
-		TArrayView<const FKeyHandle> SelectedHandles = Set.Value.AsArray();
-		TArray<FKeyHandle> NonMutableArray = TArray<FKeyHandle>(SelectedHandles.GetData(), SelectedHandles.Num());
-		
-		for (const FKeyHandle& Handle : NonMutableArray)
-		{
-			// Check to see if our curve model contains this handle still.
-			if (!KeyHandles.Contains(Handle))
-			{
-				Selection.Remove(Set.Key, ECurvePointType::Key, Handle);
-			}
-		}
-	}
 }
 
 void FCurveEditor::PostRedo(bool bSuccess)
diff --git a/Engine/Source/Editor/CurveEditor/Private/Modification/Keys/Diff/MultiGenericCommandAppender.cpp b/Engine/Source/Editor/CurveEditor/Private/Modification/Keys/Diff/MultiGenericCommandAppender.cpp
--- a/Engine/Source/Editor/CurveEditor/Private/Modification/Keys/Diff/MultiGenericCommandAppender.cpp
+++ b/Engine/Source/Editor/CurveEditor/Private/Modification/Keys/Diff/MultiGenericCommandAppender.cpp
@@ -53,11 +53,8 @@ void FMultiGenericCommandAppender::PostProcessChanges()
 	// There is no package so we don't need to worry about any dirty flags.
 	if (ChangesToCurvesWithoutUObject)
 	{
-		TUniquePtr<FGenericCurveChangeCommand> PerCurveCommand = MakeUnique<FGenericCurveChangeCommand>(
-			CurveEditor, MoveTemp(ChangesToCurvesWithoutUObject)
-			);
-		AllocationSize += PerCurveCommand->GetSize();
-		CurveEditor->GetTransactionManager()->AppendChange(MoveTemp(PerCurveCommand));
+		AllocationSize += ChangesToCurvesWithoutUObject.GetAllocatedSize() + sizeof(FGenericCurveChangeCommand);
+		CurveEditor->GetTransactionManager()->AppendCurveChange(MoveTemp(ChangesToCurvesWithoutUObject));
 	}
 }
 
@@ -75,9 +72,8 @@ void FMultiGenericCommandAppender::ProcessDiff(const FCurveModelID& InCurveId, F
 	
 	if (UObject* OwningObject = CurveModel->GetOwningObject())
 	{
-		TUniquePtr<FGenericCurveChangeCommand> PerCurveCommand = MakeUnique<FGenericCurveChangeCommand>(CurveEditor, MoveTemp(InData));
-		AllocationSize += PerCurveCommand->GetSize();
-		CurveEditor->GetTransactionManager()->AppendChange(OwningObject, MoveTemp(PerCurveCommand));
+		AllocationSize += InData.GetAllocatedSize() + sizeof(FGenericCurveChangeCommand);
+		CurveEditor->GetTransactionManager()->AppendCurveChange(OwningObject, MoveTemp(InData));
 	}
 	else
 	{
diff --git a/Engine/Source/Editor/CurveEditor/Private/Modification/TransactionManager.cpp b/Engine/Source/Editor/CurveEditor/Private/Modification/TransactionManager.cpp
--- a/Engine/Source/Editor/CurveEditor/Private/Modification/TransactionManager.cpp
+++ b/Engine/Source/Editor/CurveEditor/Private/Modification/TransactionManager.cpp
@@ -4,7 +4,9 @@
 
 #include "CurveEditorTransactionObject.h"
 #include "ScopedTransaction.h"
+#include "Keys/GenericCurveChangeCommand.h"
 #include "Misc/ITransaction.h"
+#include "Modification/Keys/Data/GenericCurveChangeData.h"
 #include "UObject/Package.h"
 #include "UObject/UObjectGlobals.h"
 
@@ -15,14 +17,42 @@ FTransactionManager::FTransactionManager(TWeakPtr<FCurveEditor> InCurveEditor)
 {
 	DummyTransactionObject->OwningCurveEditor = MoveTemp(InCurveEditor);
 }
-	
+
+void FTransactionManager::AppendCurveChange(UObject* InObject, FGenericCurveChangeData InChange, const FText& InDescription) const
+{
+	if (const TSharedPtr<FCurveEditor> CurveEditorPin = DummyTransactionObject->OwningCurveEditor.Pin())
+	{
+		// Problem: OnCurvesChangedDelegate should be called after calling StoreUndo in case any follow-up changes by the subscribers. Otherwise, we
+		// may cause some subtle ordering issues during undo / redo.
+		// StoreUndo takes ownership over the TUniquePtr so we cannot access the FGenericCurveChangeData from the FGenericCurveChangeCommand.
+		// Copying here is a bit unfortunate, but it's the easiest solution without having to implement some advanced queueing system.
+		const FGenericCurveChangeData CopiedChange = OnCurvesChangedDelegate.IsBound() ? InChange : FGenericCurveChangeData();
+		
+		// Start a transaction...
+		const FScopedTransaction Transaction(InDescription);
+		AppendChange(
+			InObject,
+			MakeUnique<FGenericCurveChangeCommand>(CurveEditorPin.ToSharedRef(), MoveTemp(InChange)),
+			InDescription
+			);
+		// ... in case any follow-up changes want to add anything to the transaction.
+		OnCurvesChangedDelegate.Broadcast(CopiedChange);
+	}
+}
+
+void FTransactionManager::AppendCurveChange(FGenericCurveChangeData InChange, const FText& InDescription) const
+{
+	AppendCurveChange(DummyTransactionObject, MoveTemp(InChange), InDescription);
+}
+
 void FTransactionManager::AppendChange(UObject* InObject, TUniquePtr<FCurveEditorCommandChange> InChange, const FText& InDescription) const
 {
-	FScopedTransaction Transaction(InDescription);
+	const FScopedTransaction Transaction(InDescription);
 	if (GUndo // GUndo can be reset by external systems, e.g. during world tear down.
 		&& ensure(InChange))
 	{
 		GUndo->StoreUndo(InObject, MoveTemp(InChange));
+		
 		OnCommandAppendedDelegate.Broadcast();
 	}
 }
diff --git a/Engine/Source/Editor/CurveEditor/Private/Selection/SelectionCleanser.cpp b/Engine/Source/Editor/CurveEditor/Private/Selection/SelectionCleanser.cpp
new file mode 100644
--- /dev/null
+++ b/Engine/Source/Editor/CurveEditor/Private/Selection/SelectionCleanser.cpp
@@ -0,0 +1,62 @@
+﻿// Copyright Epic Games, Inc. All Rights Reserved.
+
+#include "SelectionCleanser.h"
+
+#include "CurveEditor.h"
+#include "SCurveEditor.h"
+#include "SelectionUtils.h"
+
+namespace UE::CurveEditor
+{
+FSelectionCleanser::FSelectionCleanser(const TSharedRef<FCurveEditor>& InCurveEditor)
+	: WeakCurveEditor(InCurveEditor)
+{
+	if (FTransactionManager* TransactionManager = InCurveEditor->GetTransactionManager(); ensure(TransactionManager))
+	{
+		TransactionManager->OnCurvesChanged().AddRaw(this, &FSelectionCleanser::OnCurvesChanged);
+	}
+}
+
+FSelectionCleanser::~FSelectionCleanser()
+{
+	const TSharedPtr<FCurveEditor> CurveEditorPin = WeakCurveEditor.Pin();
+	if (!CurveEditorPin)
+	{
+		return;
+	}
+	
+	if (FTransactionManager* TransactionManager = CurveEditorPin->GetTransactionManager())
+	{
+		TransactionManager->OnCurvesChanged().RemoveAll(this);
+	}
+}
+
+void FSelectionCleanser::PostUndo(bool bSuccess)
+{
+	if (const TSharedPtr<FCurveEditor> CurveEditorPin = WeakCurveEditor.Pin())
+	{
+		// If you create keys and then undo them, they may remain in the selection set.
+		// Code that mutates keys should handle this case by using FScopedCurveChange & FScopedSelectionChange, but we'll clean up the selection
+		// anyway just in case.
+		// Stale keys presents issues with context menus and other things that are activated when there is a selection set.
+		// To fix this, we have to loop through all of our curve models, and re-select only the key handles that were
+		// previously selected that still exist. Ugly, but reasonably functional.
+		const bool bHadStaleKeys = CleanseSelection(CurveEditorPin.ToSharedRef(), CurveEditorPin->Selection) == ECleanseResult::HadStaleKeys;
+		UE_CLOG(bHadStaleKeys, LogCurveEditor, Warning, TEXT("Curve Editor had stale keys after undo & redo. Investigate."));
+	}
+}
+
+void FSelectionCleanser::OnCurvesChanged(const FGenericCurveChangeData& InChange) const
+{
+	const TSharedPtr<FCurveEditor> CurveEditorPin = WeakCurveEditor.Pin();
+	if (!CurveEditorPin || !InChange.RemoveKeysData.HasChanges())
+	{
+		return;
+	}
+	
+	for (const TPair<FCurveModelID, FCurveKeyData>& Pair : InChange.RemoveKeysData.SavedCurveState)
+	{
+		CleanseSelection(CurveEditorPin.ToSharedRef(), CurveEditorPin->Selection, MakeConstArrayView({ Pair.Key }));
+	}
+}
+}
diff --git a/Engine/Source/Editor/CurveEditor/Private/Selection/SelectionCleanser.h b/Engine/Source/Editor/CurveEditor/Private/Selection/SelectionCleanser.h
new file mode 100644
--- /dev/null
+++ b/Engine/Source/Editor/CurveEditor/Private/Selection/SelectionCleanser.h
@@ -0,0 +1,47 @@
+﻿// Copyright Epic Games, Inc. All Rights Reserved.
+
+#pragma once
+
+#include "EditorUndoClient.h"
+#include "Modification/Keys/Diff/AppendPerCurveGenericChangeVisitor.h"
+#include "Templates/UnrealTemplate.h"
+
+class FCurveEditor;
+
+namespace UE::CurveEditor
+{
+/**
+ * Makes sure that the FCurveEditor::Selection does not contain any stale keys, i.e. to prevent that the selection contains FKeyHandles and
+ * FCurveModelIds that do not exist in the FCurveEditor.
+ * 
+ * Stale keys presents issues with context menus and other things that are activated when there is a selection set.
+ *
+ * Cleansing happens when:
+ * - An undo / redo operation is performed (regardless whether it directly affects FCurveEditor - the data underlying the FCurveModel may change)
+ * - An undoable action is recorded to FTransactionManager.
+ * Effectively, this means that if you use FScopedCurveChange to make changes to curves, this manager will handle cleansing selection.
+ */
+class FSelectionCleanser
+	: public FNoncopyable
+	, public FEditorUndoClient
+{
+public:
+	
+	explicit FSelectionCleanser(const TSharedRef<FCurveEditor>& InCurveEditor);
+	~FSelectionCleanser();
+
+	//~ Begin FEditorUndoClient Interface
+	virtual void PostUndo(bool bSuccess) override;
+	virtual void PostRedo(bool bSuccess) override { PostUndo(bSuccess); }
+	//~ End FEditorUndoClient Interface
+
+private:
+
+	/** Instance that owns us. Used to check which curves are currently present. */
+	const TWeakPtr<FCurveEditor> WeakCurveEditor;
+
+	/** Invoked when a (transactional) change happens to curves. */
+	void OnCurvesChanged(const FGenericCurveChangeData& InChange) const;
+};
+}
+
diff --git a/Engine/Source/Editor/CurveEditor/Private/Selection/SelectionUtils.cpp b/Engine/Source/Editor/CurveEditor/Private/Selection/SelectionUtils.cpp
new file mode 100644
--- /dev/null
+++ b/Engine/Source/Editor/CurveEditor/Private/Selection/SelectionUtils.cpp
@@ -0,0 +1,87 @@
+﻿// Copyright Epic Games, Inc. All Rights Reserved.
+
+#include "SelectionUtils.h"
+
+#include "CurveEditor.h"
+#include "CurveEditorTrace.h"
+#include "SCurveEditor.h"
+#include "Selection/ScopedSelectionChangeEventSuppression.h"
+
+namespace UE::CurveEditor
+{
+namespace SelectionUtilsPrivate
+{
+static bool CleanseSingle(
+	const TSharedRef<FCurveEditor>& InCurveEditor, FCurveEditorSelection& InSelection,
+	const FCurveModelID& InModelId, const FKeyHandleSet& InSet
+	)
+{
+	SCOPED_CURVE_EDITOR_TRACE(CleanseSingle);
+	
+	FCurveModel* CurveModel = InCurveEditor->FindCurve(InModelId);
+
+	// If the entire curve was removed, just dump that out of the selection set.
+	if (!CurveModel)
+	{
+		InSelection.Remove(InModelId);
+		return true;
+	}
+	// Get all of the key handles from this curve.
+	const TArray<FKeyHandle> KeyHandles = CurveModel->GetAllKeys();
+
+	// The set handles will be mutated as we remove things so we need a copy that we can iterate through.
+	const TConstArrayView<FKeyHandle> SelectedHandles = InSet.AsArray();
+	const TArray<FKeyHandle> NonMutableArray(SelectedHandles.GetData(), SelectedHandles.Num());
+
+	bool bContainedStaleKeys = false;
+	for (const FKeyHandle& Handle : NonMutableArray)
+	{
+		// Check to see if our curve model contains this handle still.
+		if (!KeyHandles.Contains(Handle))
+		{
+			InSelection.Remove(InModelId, ECurvePointType::Key, Handle);
+			bContainedStaleKeys = true;
+		}
+	}
+	return bContainedStaleKeys;
+}
+}
+	
+ECleanseResult CleanseSelection(
+	const TSharedRef<FCurveEditor>& InCurveEditor, FCurveEditorSelection& InSelection
+	)
+{
+	SCOPED_CURVE_EDITOR_TRACE(CleanseSelection);
+	// Minor optimization to reduce selection changed calls
+	const FScopedSelectionChangeEventSuppression SupressSelectionEvents(InCurveEditor);
+
+	bool bHadStaleCurvesOrKeys = false;
+	TMap<FCurveModelID, FKeyHandleSet> SelectionSet = InSelection.GetAll();
+	for (const TPair<FCurveModelID, FKeyHandleSet>& Set : SelectionSet)
+	{
+		bHadStaleCurvesOrKeys |= SelectionUtilsPrivate::CleanseSingle(InCurveEditor, InSelection, Set.Key, Set.Value);
+	}
+
+	return bHadStaleCurvesOrKeys ? ECleanseResult::HadStaleKeys : ECleanseResult::NoStaleKeys;
+}
+
+ECleanseResult CleanseSelection(
+	const TSharedRef<FCurveEditor>& InCurveEditor, FCurveEditorSelection& InSelection, TConstArrayView<FCurveModelID> InOnlyTheseCurves
+	)
+{
+	SCOPED_CURVE_EDITOR_TRACE(CleanseSelection);
+	// Minor optimization to reduce selection changed calls
+	const FScopedSelectionChangeEventSuppression SupressSelectionEvents(InCurveEditor);
+	
+	bool bHadStaleCurvesOrKeys = false;
+	for (const FCurveModelID& CurveId : InOnlyTheseCurves)
+	{
+		if (const FKeyHandleSet* Set = InSelection.FindForCurve(CurveId))
+		{
+			bHadStaleCurvesOrKeys |= SelectionUtilsPrivate::CleanseSingle(InCurveEditor, InSelection, CurveId, *Set);
+		}
+	}
+	
+	return bHadStaleCurvesOrKeys ? ECleanseResult::HadStaleKeys : ECleanseResult::NoStaleKeys;
+}
+}
diff --git a/Engine/Source/Editor/CurveEditor/Private/Selection/SelectionUtils.h b/Engine/Source/Editor/CurveEditor/Private/Selection/SelectionUtils.h
new file mode 100644
--- /dev/null
+++ b/Engine/Source/Editor/CurveEditor/Private/Selection/SelectionUtils.h
@@ -0,0 +1,23 @@
+﻿// Copyright Epic Games, Inc. All Rights Reserved.
+
+#pragma once
+
+#include "Containers/ContainersFwd.h"
+#include "Templates/SharedPointerFwd.h"
+
+struct FCurveModelID;
+struct FCurveEditorSelection;
+class FCurveEditor;
+
+namespace UE::CurveEditor
+{
+enum class ECleanseResult : uint8 { NoStaleKeys, HadStaleKeys };
+	
+/** Removes all curve models and keys from InSelection that do not exist in InCurveEditor. */
+ECleanseResult CleanseSelection(const TSharedRef<FCurveEditor>& InCurveEditor, FCurveEditorSelection& InSelection);
+
+/** Overload that will only remove 1. those FKeyHandles belonging to curves in InOnlyTheseCurves, 2. the curves in InOnlyTheseCurves. */
+ECleanseResult CleanseSelection(const TSharedRef<FCurveEditor>& InCurveEditor, FCurveEditorSelection& InSelection,
+	TConstArrayView<FCurveModelID> InOnlyTheseCurves
+	);
+}
diff --git a/Engine/Source/Editor/CurveEditor/Public/CurveEditor.h b/Engine/Source/Editor/CurveEditor/Public/CurveEditor.h
--- a/Engine/Source/Editor/CurveEditor/Public/CurveEditor.h
+++ b/Engine/Source/Editor/CurveEditor/Public/CurveEditor.h
@@ -57,6 +57,7 @@ struct FCurveDrawParams;
 struct FCurveEditorInitParams;
 struct FCurveEditorSnapMetrics;
 struct FGeometry;
+namespace UE::CurveEditor { class FSelectionCleanser; }
 
 DECLARE_DELEGATE_OneParam(FOnSetBoolean, bool)
 DECLARE_MULTICAST_DELEGATE_OneParam(FOnActiveToolChanged, FCurveEditorToolID)
@@ -761,8 +762,13 @@ protected:
 
 	/** Orchestrates transaction that affect this curve editor. */
 	TUniquePtr<UE::CurveEditor::FTransactionManager> TransactionManager;
+	
+	/** Makes sure that selection does not reference stale FKeyHandles or FCurveModelIds, i.e. after undo or an operation has removed keys. */
+	TPimplPtr<UE::CurveEditor::FSelectionCleanser> SelectionCleanser;
 
 private:
+
+	
 	/** Cached physical size of the panel representing this editor */
 	FVector2D CachedPhysicalSize;
 
diff --git a/Engine/Source/Editor/CurveEditor/Public/Modification/TransactionManager.h b/Engine/Source/Editor/CurveEditor/Public/Modification/TransactionManager.h
--- a/Engine/Source/Editor/CurveEditor/Public/Modification/TransactionManager.h
+++ b/Engine/Source/Editor/CurveEditor/Public/Modification/TransactionManager.h
@@ -15,14 +15,24 @@ class UCurveEditorTransactionObject;
 
 namespace UE::CurveEditor
 {
+struct FGenericCurveChangeData;
+	
 /** Interacts with the transaction system to enable command-based actions for curve editor. */
 class FTransactionManager : public FNoncopyable, public FGCObject
 {
 public:
 
 	UE_API explicit FTransactionManager(TWeakPtr<FCurveEditor> InCurveEditor);
-	
+
+	/** Appends a changes made to curves. It is associated with the current transaction.  */
+	void AppendCurveChange(UObject* InObject, FGenericCurveChangeData InChange, const FText& InDescription = FText::GetEmpty()) const;
 
+	/**
+	 * Appends a changes made to curves. It is associated with the current transaction. 
+	 * Use this version when there is no UObject the change can be associated with - otherwise use the UObject overload.
+	 */
+	void AppendCurveChange(FGenericCurveChangeData InChange, const FText& InDescription = FText::GetEmpty()) const;
+	
 	/**
 	 * Appends a change to the current transaction and associates it with an UObject.
 	 * 
@@ -51,6 +61,10 @@ public:
 	DECLARE_MULTICAST_DELEGATE(FOnCommandAppended);
 	/** Invoked after a command has been appended to the undo stack. */
 	FOnCommandAppended& OnCommandAppended() { return OnCommandAppendedDelegate; }
+	
+	DECLARE_MULTICAST_DELEGATE_OneParam(FOnCurvesChanged, const FGenericCurveChangeData&)
+	/** Invoked when transaction changes to curves are made. */
+	FOnCurvesChanged& OnCurvesChanged() { return OnCurvesChangedDelegate; }
 
 	//~ Begin FGCObject Interface
 	UE_API virtual void AddReferencedObjects(FReferenceCollector& Collector) override;
@@ -64,6 +78,9 @@ private:
 
 	/** Invoked after a command has been appended to the undo stack. */
 	FOnCommandAppended OnCommandAppendedDelegate;
+
+	/** Invoked when transaction changes to curves are made. */
+	FOnCurvesChanged OnCurvesChangedDelegate;
 };
 }
 
