Curve Editor: Fix forgetting to register selection cleanser for undo / redo

diff --git a/Engine/Source/Editor/CurveEditor/Private/Selection/SelectionCleanser.cpp b/Engine/Source/Editor/CurveEditor/Private/Selection/SelectionCleanser.cpp
--- a/Engine/Source/Editor/CurveEditor/Private/Selection/SelectionCleanser.cpp
+++ b/Engine/Source/Editor/CurveEditor/Private/Selection/SelectionCleanser.cpp
@@ -3,6 +3,7 @@
 #include "SelectionCleanser.h"
 
 #include "CurveEditor.h"
+#include "Editor.h"
 #include "SCurveEditor.h"
 #include "SelectionUtils.h"
 
@@ -15,10 +16,20 @@ FSelectionCleanser::FSelectionCleanser(const TSharedRef<FCurveEditor>& InCurveEd
 	{
 		TransactionManager->OnCurvesChanged().AddRaw(this, &FSelectionCleanser::OnCurvesChanged);
 	}
+
+	if (GEditor)
+	{
+		GEditor->RegisterForUndo(this);
+	}
 }
 
 FSelectionCleanser::~FSelectionCleanser()
 {
+	if (GEditor)
+	{
+		GEditor->UnregisterForUndo(this);
+	}
+	
 	const TSharedPtr<FCurveEditor> CurveEditorPin = WeakCurveEditor.Pin();
 	if (!CurveEditorPin)
 	{
