Curve Editor: Fix an issue where curves  flicker hovered and unhovered while editing

diff --git a/Engine/Source/Editor/CurveEditor/Private/Views/SInteractiveCurveEditorView.cpp b/Engine/Source/Editor/CurveEditor/Private/Views/SInteractiveCurveEditorView.cpp
--- a/Engine/Source/Editor/CurveEditor/Private/Views/SInteractiveCurveEditorView.cpp
+++ b/Engine/Source/Editor/CurveEditor/Private/Views/SInteractiveCurveEditorView.cpp
@@ -1185,6 +1185,12 @@ bool SInteractiveCurveEditorView::GetCurveWithinWidgetRange(const FSlateRect& Wi
 
 void SInteractiveCurveEditorView::UpdateCurveProximities(FVector2D MousePixel)
 {
+	if (DragOperation.IsSet())
+	{
+		// Don't update while dragging, assume the mouse remains over the hovered curve.
+		return;
+	}
+
 	TOptional<FCurveModelID> PreviouslyHovered = GetHoveredCurve();
 
 	CurveProximities.Reset();
@@ -1231,8 +1237,11 @@ void SInteractiveCurveEditorView::UpdateCurveProximities(FVector2D MousePixel)
 			MinPos.Y = CurveSpace.ValueToScreen(MinValue);
 			MaxPos.Y = CurveSpace.ValueToScreen(MaxValue);
 
-			float Distance = (FMath::ClosestPointOnSegment2D(MousePixel, MinPos, MaxPos) - MousePixel).Size();
-			CurveProximities.Add(MakeTuple(Pair.Key, Distance));
+			const float Distance = (FMath::ClosestPointOnSegment2D(MousePixel, MinPos, MaxPos) - MousePixel).Size();
+			if (Distance < CurveViewConstants::HoverProximityThresholdPx)
+			{
+				CurveProximities.Add(MakeTuple(Pair.Key, Distance));
+			}
 		}
 	}
 
@@ -1240,7 +1249,7 @@ void SInteractiveCurveEditorView::UpdateCurveProximities(FVector2D MousePixel)
 
 	TOptional<FCurveModelID> NewHovered;
 	// Also, set the cached tooltips if dragging because the curve proximity might not be updated during the drag
-	if (CurveProximities.Num() > 0 && (CurveProximities[0].Get<1>() < CurveViewConstants::HoverProximityThresholdPx || DragOperation.IsSet()))
+	if (CurveProximities.Num() > 0)
 	{
 		NewHovered = CurveProximities[0].Get<0>();
 
@@ -1265,7 +1274,7 @@ void SInteractiveCurveEditorView::UpdateCurveProximities(FVector2D MousePixel)
 
 	if (PreviouslyHovered != NewHovered)
 	{
-		bNeedsRefresh = true;
+		RefreshRetainer();
 	}
 }
 
@@ -1285,15 +1294,26 @@ FText SInteractiveCurveEditorView::FormatToolTipValue(const FCurveModel& CurveMo
 }
 
 void SInteractiveCurveEditorView::OnMouseEnter(const FGeometry& MyGeometry, const FPointerEvent& MouseEvent)
-{
-	//CurveProximities.Reset();
+{	
+	// Unless a drag drop mouse move is reentering, make sure no curve is initially hovered
+	if (!DragOperation.IsSet())
+	{
+		CurveProximities.Reset();
+		RefreshRetainer();
+	}
+
 	SCompoundWidget::OnMouseEnter(MyGeometry, MouseEvent);
 }
 
 void SInteractiveCurveEditorView::OnMouseLeave(const FPointerEvent& MouseEvent)
 {
-	// Don't allow hover highlights when we've exited this view as clicking won't be routed to us to select it anyways.
-	CurveProximities.Reset();
+	// Unless a drag drop mouse move is leaving, make sure no curve shows hovered
+	if (!DragOperation.IsSet())
+	{
+		CurveProximities.Reset();
+		RefreshRetainer();
+	}
+
 	SCompoundWidget::OnMouseLeave(MouseEvent);
 }
 
@@ -1397,7 +1417,8 @@ FReply SInteractiveCurveEditorView::OnMouseWheel(const FGeometry& MyGeometry, co
 
 TOptional<FCurveModelID> SInteractiveCurveEditorView::GetHoveredCurve() const
 {
-	if (CurveProximities.Num() > 0 && CurveProximities[0].Get<1>() < CurveViewConstants::HoverProximityThresholdPx)
+	if (CurveProximities.Num() > 0 && 
+		CurveProximities[0].Get<0>().IsValid())
 	{
 		return CurveProximities[0].Get<0>();
 	}
diff --git a/Engine/Source/Editor/CurveEditor/Public/Views/SInteractiveCurveEditorView.h b/Engine/Source/Editor/CurveEditor/Public/Views/SInteractiveCurveEditorView.h
--- a/Engine/Source/Editor/CurveEditor/Public/Views/SInteractiveCurveEditorView.h
+++ b/Engine/Source/Editor/CurveEditor/Public/Views/SInteractiveCurveEditorView.h
@@ -161,7 +161,7 @@ private:
 
 	UE_API void CreateContextMenu(const FGeometry& MyGeometry, const FPointerEvent& MouseEvent);
 
-	/** Updates our distance to all of the curves we represent. */
+	/** Updates our distance to the curve neareast to the cursor. */
 	UE_API void UpdateCurveProximities(FVector2D MousePixel);
 
 	UE_API TOptional<FCurvePointHandle> HitPoint(FVector2D MousePixel) const;
@@ -230,7 +230,7 @@ private:
 
 	TOptional<FCachedToolTipData> CachedToolTipData;
 
-	/** Array of curve proximities in slate units that's updated on mouse move */
+	/** Array of curve proximities in slate units that's updated on mouse move. Does not necessarily hold all curves. */
 	TArray<TTuple<FCurveModelID, float>> CurveProximities;
 
 	/** Track if we have a context menu active. Used to supress hover updates as it causes flickers in the CanExecute bindings. */
