Curve Editor: Mend comment in SInteractiveCurveEditorView::UpdateCurveProximities

diff --git a/Engine/Source/Editor/CurveEditor/Private/Views/SInteractiveCurveEditorView.cpp b/Engine/Source/Editor/CurveEditor/Private/Views/SInteractiveCurveEditorView.cpp
--- a/Engine/Source/Editor/CurveEditor/Private/Views/SInteractiveCurveEditorView.cpp
+++ b/Engine/Source/Editor/CurveEditor/Private/Views/SInteractiveCurveEditorView.cpp
@@ -1248,11 +1248,11 @@ void SInteractiveCurveEditorView::UpdateCurveProximities(FVector2D MousePixel)
 	Algo::SortBy(CurveProximities, [](TTuple<FCurveModelID, float> In) { return In.Get<1>(); });
 
 	TOptional<FCurveModelID> NewHovered;
-	// Also, set the cached tooltips if dragging because the curve proximity might not be updated during the drag
 	if (CurveProximities.Num() > 0)
 	{
 		NewHovered = CurveProximities[0].Get<0>();
 
+		// Update tooltips with data from the current mouse position if there is a hovered curve
 		const FCurveModel* HoveredCurve = CurveEditor->FindCurve(CurveProximities[0].Get<0>());
 		if (HoveredCurve)
 		{
