Curve Editor: Paste keys don't undo.

diff --git a/Engine/Source/Editor/CurveEditor/Private/CurveEditor.cpp b/Engine/Source/Editor/CurveEditor/Private/CurveEditor.cpp
--- a/Engine/Source/Editor/CurveEditor/Private/CurveEditor.cpp
+++ b/Engine/Source/Editor/CurveEditor/Private/CurveEditor.cpp
@@ -1866,10 +1866,9 @@ void FCurveEditor::PasteKeys(UE::CurveEditor::FKeyPasteArgs InArgs)
 		}
 	}
 	
-	const FScopedTransaction Transaction(LOCTEXT("PasteKeys", "Paste Keys")); // Selection and key change should be part of same undo / redo op.
+	const FScopedTransaction Transaction(InArgs.OverrideTransactionName); // Selection and key change should be part of same undo / redo op.
 	const FScopedSelectionChange SelectionChange(SharedThis(this));
-	constexpr ECurveChangeFlags Flags = ECurveChangeFlags::MoveKeysAndRemoveStackedKeys | ECurveChangeFlags::RemoveKeys | ECurveChangeFlags::KeyAttributes;
-	const FScopedCurveChange KeyChange(FCurvesSnapshotBuilder(SharedThis(this), TargetCurves, Flags));
+	const FScopedCurveChange KeyChange(FCurvesSnapshotBuilder(SharedThis(this), TargetCurves, ECurveChangeFlags::KeyData));
 	Selection.Clear();
 
 	// Two simple cases, 1 to 1 and 1 to many.
