Curve Editor: Correctly implement ToString for undo commands.

diff --git a/Engine/Source/Editor/CurveEditor/Private/Modification/CurveEditorCommandChange.cpp b/Engine/Source/Editor/CurveEditor/Private/Modification/CurveEditorCommandChange.cpp
--- a/Engine/Source/Editor/CurveEditor/Private/Modification/CurveEditorCommandChange.cpp
+++ b/Engine/Source/Editor/CurveEditor/Private/Modification/CurveEditorCommandChange.cpp
@@ -14,7 +14,7 @@ bool FCurveEditorCommandChange::HasExpired(UObject* InObject) const
 
 FString FCurveEditorCommandChange::ToString() const
 {
-	return TEXT("FSelectionChangeCommand");
+	return TEXT("FCurveEditorCommandChange");
 }
 
 TSharedPtr<FCurveEditor> FCurveEditorCommandChange::GetCurveEditor() const
diff --git a/Engine/Source/Editor/CurveEditor/Private/Modification/Keys/GenericCurveChangeCommand.cpp b/Engine/Source/Editor/CurveEditor/Private/Modification/Keys/GenericCurveChangeCommand.cpp
--- a/Engine/Source/Editor/CurveEditor/Private/Modification/Keys/GenericCurveChangeCommand.cpp
+++ b/Engine/Source/Editor/CurveEditor/Private/Modification/Keys/GenericCurveChangeCommand.cpp
@@ -27,4 +27,9 @@ SIZE_T FGenericCurveChangeCommand::GetSize() const
 {
 	return DeltaChange.GetAllocatedSize() + sizeof(FGenericCurveChangeCommand);
 }
+
+FString FGenericCurveChangeCommand::ToString() const
+{
+	return TEXT("FGenericCurveChangeCommand");
+}
 }
diff --git a/Engine/Source/Editor/CurveEditor/Private/Modification/Keys/GenericCurveChangeCommand.h b/Engine/Source/Editor/CurveEditor/Private/Modification/Keys/GenericCurveChangeCommand.h
--- a/Engine/Source/Editor/CurveEditor/Private/Modification/Keys/GenericCurveChangeCommand.h
+++ b/Engine/Source/Editor/CurveEditor/Private/Modification/Keys/GenericCurveChangeCommand.h
@@ -29,6 +29,7 @@ public:
 	virtual void Apply(UObject* Object) override;
 	virtual void Revert(UObject* Object) override;
 	virtual SIZE_T GetSize() const override;
+	virtual FString ToString() const override;
 	//~ End FChange Interface
 
 private:
diff --git a/Engine/Source/Editor/CurveEditor/Private/Modification/Selection/SelectionChangeCommand.cpp b/Engine/Source/Editor/CurveEditor/Private/Modification/Selection/SelectionChangeCommand.cpp
--- a/Engine/Source/Editor/CurveEditor/Private/Modification/Selection/SelectionChangeCommand.cpp
+++ b/Engine/Source/Editor/CurveEditor/Private/Modification/Selection/SelectionChangeCommand.cpp
@@ -30,4 +30,9 @@ SIZE_T FSelectionChangeCommand::GetSize() const
 {
 	return sizeof(FSelectionChangeCommand) + DeltaChange.GetAllocatedSize();
 }
+
+FString FSelectionChangeCommand::ToString() const
+{
+	return TEXT("FSelectionChangeCommand");
+}
 }
diff --git a/Engine/Source/Editor/CurveEditor/Private/Modification/Selection/SelectionChangeCommand.h b/Engine/Source/Editor/CurveEditor/Private/Modification/Selection/SelectionChangeCommand.h
--- a/Engine/Source/Editor/CurveEditor/Private/Modification/Selection/SelectionChangeCommand.h
+++ b/Engine/Source/Editor/CurveEditor/Private/Modification/Selection/SelectionChangeCommand.h
@@ -21,6 +21,7 @@ public:
 	virtual void Apply(UObject* Object) override;
 	virtual void Revert(UObject* Object) override;
 	virtual SIZE_T GetSize() const override;
+	virtual FString ToString() const override;
 	//~ End FChange Interface
 
 private:
