Curve Editor: Fix not being able to undo adding keys to a curve when that keys previously had no keys at all.

This was a logic error. The diffing code for detecting whether any keys have been aded would not run if the curves had no keys to begin with... assuming that the data was not initialized.

diff --git a/Engine/Source/Editor/CurveEditor/Private/Modification/Keys/CurveChangeDiff.cpp b/Engine/Source/Editor/CurveEditor/Private/Modification/Keys/CurveChangeDiff.cpp
--- a/Engine/Source/Editor/CurveEditor/Private/Modification/Keys/CurveChangeDiff.cpp
+++ b/Engine/Source/Editor/CurveEditor/Private/Modification/Keys/CurveChangeDiff.cpp
@@ -47,7 +47,7 @@ static void HandleOp_AddedKeys(
 	ISingleCurveChangeVisitor& InCurveChangeVisitor
 )
 {
-	if (!EnumHasAnyFlags(Operations, ECurveChangeFlags::AddKeys) || !SavedState)
+	if (!EnumHasAnyFlags(Operations, ECurveChangeFlags::AddKeys))
 	{
 		return;
 	}
