Curve Editor: Introduce EnterChangeScope and ExitChangeScope. This allows FChannelCurveModel to postbone OnChannelChanged event until after the undo commands of FScopedCurveChange have been recorded.

diff --git a/Engine/Plugins/Editor/CurveEditorTools/Source/CurveEditorTools/Private/Tools/Lattice/CurveEditorLatticeTool.cpp b/Engine/Plugins/Editor/CurveEditorTools/Source/CurveEditorTools/Private/Tools/Lattice/CurveEditorLatticeTool.cpp
--- a/Engine/Plugins/Editor/CurveEditorTools/Source/CurveEditorTools/Private/Tools/Lattice/CurveEditorLatticeTool.cpp
+++ b/Engine/Plugins/Editor/CurveEditorTools/Source/CurveEditorTools/Private/Tools/Lattice/CurveEditorLatticeTool.cpp
@@ -183,6 +183,9 @@ FReply FCurveEditorLatticeTool::OnMouseButtonUp(
 {
 	if (DelayedDragData.IsSet())
 	{
+		// Resetting DelayedDragData may cause FScopedCurveChange to invoke the curves' modified events... which would cause us to recompute our lattice grid.
+		const TGuardValue Guard(LatticeDeformerState->bIsModifyingCurves, true);
+		
 		FReply Reply = FReply::Handled();
 		if (DelayedDragData->DelayedDrag.IsDragging() && ensure(DelayedDragData->ActiveOperation))
 		{
@@ -190,6 +193,7 @@ FReply FCurveEditorLatticeTool::OnMouseButtonUp(
 			// Only return handled if we actually started a drag
 			Reply.ReleaseMouseCapture();
 		}
+		
 		DelayedDragData.Reset();
 		return Reply;
 	}
diff --git a/Engine/Source/Editor/CurveEditor/Private/CurveModel.cpp b/Engine/Source/Editor/CurveEditor/Private/CurveModel.cpp
--- a/Engine/Source/Editor/CurveEditor/Private/CurveModel.cpp
+++ b/Engine/Source/Editor/CurveEditor/Private/CurveModel.cpp
@@ -106,3 +106,26 @@ TOptional<FKeyHandle> FCurveModel::AddKey(const FKeyPosition& NewKeyPosition, co
 	AddKeys(TArrayView<const FKeyPosition>(&NewKeyPosition, 1), TArrayView<const FKeyAttributes>(&InAttributes, 1), &Handles);
 	return Handle;
 }
+
+void FCurveModel::OpenChangeScope()
+{
+	++ChangeScopeCounter;
+	if (ChangeScopeCounter == 1)
+	{
+		OnOpenRootChangeScope();
+	}
+}
+
+void FCurveModel::CloseChangeScope()
+{
+	if (!ensure(ChangeScopeCounter > 0))
+	{
+		return;
+	}
+
+	--ChangeScopeCounter;
+	if (ChangeScopeCounter == 0)
+	{
+		OnCloseRootChangeScope();
+	}
+}
diff --git a/Engine/Source/Editor/CurveEditor/Private/Modification/Utils/ModifyCurveObjectDetector.cpp b/Engine/Source/Editor/CurveEditor/Private/Modification/Utils/ModifyCurveObjectDetector.cpp
new file mode 100644
--- /dev/null
+++ b/Engine/Source/Editor/CurveEditor/Private/Modification/Utils/ModifyCurveObjectDetector.cpp
@@ -0,0 +1,65 @@
+﻿// Copyright Epic Games, Inc. All Rights Reserved.
+
+#include "ModifyCurveObjectDetector.h"
+
+#include "Algo/AnyOf.h"
+#include "CurveEditor.h"
+#include "Modification/Keys/Data/CurvesSnapshot.h"
+#include "SCurveEditor.h"
+
+namespace UE::CurveEditor
+{
+FModifyCurveObjectDetector::FModifyCurveObjectDetector(
+	TWeakPtr<FCurveEditor> InCurveEditor, EModifyDetectionLevel InLevel, FString InWarningMessage, TArray<FCurveModelID> InTrackedCurves
+	)
+	: WeakCurveEditor(MoveTemp(InCurveEditor))
+	, TrackedCurves(MoveTemp(InTrackedCurves))
+	, DetectionLevel(InLevel)
+	, WarningMessage(MoveTemp(InWarningMessage))
+{
+	FCoreUObjectDelegates::OnObjectModified.AddRaw(this, &FModifyCurveObjectDetector::OnObjectModified);
+}
+
+FModifyCurveObjectDetector::FModifyCurveObjectDetector(
+	TWeakPtr<FCurveEditor> InCurveEditor, EModifyDetectionLevel InLevel, FString InWarningMessage, const FCurvesSnapshot& InSnapshot
+	)
+	: FModifyCurveObjectDetector(
+		MoveTemp(InCurveEditor),
+		InLevel,
+		MoveTemp(InWarningMessage),
+		[&InSnapshot]
+		{
+			TArray<FCurveModelID> Curves;
+			InSnapshot.CurveData.GenerateKeyArray(Curves);
+			return Curves;
+		}())
+{}
+
+FModifyCurveObjectDetector::~FModifyCurveObjectDetector()
+{
+	FCoreUObjectDelegates::OnObjectModified.RemoveAll(this);
+}
+
+void FModifyCurveObjectDetector::OnObjectModified(UObject* Object) const
+{
+	const TSharedPtr<FCurveEditor> CurveEditorPin = WeakCurveEditor.Pin();
+	if (!Object || !CurveEditorPin || DetectionLevel == EModifyDetectionLevel::None)
+	{
+		return;
+	}
+
+	const int32 bWasModified = Algo::AnyOf(CurveEditorPin->GetCurves(), [Object](const TPair<FCurveModelID, TUniquePtr<FCurveModel>>& Pair)
+	{
+		return Pair.Value && Pair.Value->GetOwningObject() == Object;
+	});
+	switch (DetectionLevel)
+	{
+	case EnsureAlways: ensureAlways(!bWasModified); [[fallthrough]]; // Will trigger double ensure once... whatever.
+	case Ensure: ensure(!bWasModified);
+	case Warn: UE_CLOG(bWasModified, LogCurveEditor, Warning, TEXT("%s"), *WarningMessage)break;
+		
+	case None: [[fallthrough]];
+	default: break;
+	}
+}
+}
diff --git a/Engine/Source/Editor/CurveEditor/Private/Modification/Utils/ModifyCurveObjectDetector.h b/Engine/Source/Editor/CurveEditor/Private/Modification/Utils/ModifyCurveObjectDetector.h
new file mode 100644
--- /dev/null
+++ b/Engine/Source/Editor/CurveEditor/Private/Modification/Utils/ModifyCurveObjectDetector.h
@@ -0,0 +1,52 @@
+﻿// Copyright Epic Games, Inc. All Rights Reserved.
+
+#pragma once
+
+#include "Templates/SharedPointer.h"
+#include "Templates/UnrealTemplate.h"
+
+struct FCurveModelID;
+class FCurveEditor;
+
+namespace UE::CurveEditor
+{
+struct FCurvesSnapshot;
+
+enum EModifyDetectionLevel : uint8 { None, Warn, Ensure, EnsureAlways };
+	
+/**
+ * Logs warnings, ensures, etc. if any FCurveModel::GetOwningObject() is modified.
+ * @see FScopedCurveChange.
+ */
+class FModifyCurveObjectDetector : public FNoncopyable
+{
+public:
+
+	explicit FModifyCurveObjectDetector(
+		TWeakPtr<FCurveEditor> InCurveEditor, EModifyDetectionLevel InLevel, FString InWarningMessage,
+		TArray<FCurveModelID> InTrackedCurves
+		);
+	explicit FModifyCurveObjectDetector(
+		TWeakPtr<FCurveEditor> InCurveEditor, EModifyDetectionLevel InLevel, FString InWarningMessage,
+		const FCurvesSnapshot& InSnapshot
+		);
+	
+	~FModifyCurveObjectDetector();
+
+private:
+
+	/** Used to get the owning objects for the curves. */
+	const TWeakPtr<FCurveEditor> WeakCurveEditor;
+
+	/** The curves that should not be modified. */
+	const TArray<FCurveModelID> TrackedCurves;
+
+	/** Determines the action to take when modification is detected. */
+	const EModifyDetectionLevel DetectionLevel;
+	/** Message to log in the warning */
+	const FString WarningMessage;
+
+	void OnObjectModified(UObject* Object) const;
+};
+}
+
diff --git a/Engine/Source/Editor/CurveEditor/Private/Modification/Utils/ScopedCurveChange.cpp b/Engine/Source/Editor/CurveEditor/Private/Modification/Utils/ScopedCurveChange.cpp
--- a/Engine/Source/Editor/CurveEditor/Private/Modification/Utils/ScopedCurveChange.cpp
+++ b/Engine/Source/Editor/CurveEditor/Private/Modification/Utils/ScopedCurveChange.cpp
@@ -4,6 +4,7 @@
 
 #include "CurveEditor.h"
 #include "CurveModel.h"
+#include "ModifyCurveObjectDetector.h"
 #include "HAL/IConsoleManager.h"
 #include "Internationalization/Text.h"
 #include "Misc/ITransaction.h"
@@ -29,17 +30,27 @@ static TAutoConsoleVariable<bool> CVarUseNewUndoSystem(
 	TEXT("CurveEditor.UseNewUndoSystem"), true, TEXT("Whether to use the new undo system. Warning: There may still be some bugs.")
 	);
 
-static void ModifyCurves(const TSharedPtr<FCurveEditor>& InCurveEditorPin, const FCurveChangeDiff& InChangeDiff)
+static TAutoConsoleVariable<int32> CVarModifyDetectionLevel(
+	TEXT("CurveEditor.UndoModifyDetectionLevel"), EModifyDetectionLevel::Ensure, TEXT("During FScopedCurveChange, detect Modify() calls made to UObjects that own FCurveModels. Modify is not safe to call during FScopedCurveChange.\n0 - No detection\n1 - log a warning\n2 - Ensure & log warning\n3 - EnsureAlways & log warning")
+);
+
+const TCHAR* ModifyDetectionWarning = TEXT(
+	"Modify() was called on curve object during FScopedCurveChange. "
+	"This can lead to subtle undo / redo bugs. Avoid the Modify() call by introducing a command or deferring the operation if possible."
+	);
+
+template<typename TCallback> requires std::is_invocable_v<TCallback, FCurveModel&>
+static void ForEachCurve(FCurveEditor& InCurveEditorPin, const FCurveChangeDiff& InChangeDiff, TCallback&& InCallback)
 {
 	for (const TPair<FCurveModelID, FCurveDiffingData>& Pair : InChangeDiff.GetSnapshot().CurveData)
 	{
-		if (FCurveModel* CurveModel = InCurveEditorPin->FindCurve(Pair.Key))
+		if (FCurveModel* CurveModel = InCurveEditorPin.FindCurve(Pair.Key))
 		{
-			CurveModel->Modify();
+			InCallback(*CurveModel);
 		}
 	}
 }
-
+	
 FScopedCurveChange::FScopedCurveChange(
 	TWeakPtr<FCurveEditor> InCurveEditor, FCurveChangeDiff InKeyChangeDiff, const FText& InDescription, EScopedKeyChangeFlags InFlags
 	)
@@ -47,12 +58,32 @@ FScopedCurveChange::FScopedCurveChange(
 	, Transaction(Description)
 	, ChangeDiff(MoveTemp(InKeyChangeDiff))
 	, Flags(InFlags)
+	// Calling Modify() on the FCurveModel's owning objects during FScopedCurveChange can cause subtle undo bugs.
+	// The intention here is to alert developers of that happening and fix it.
+	// @see class doc string
+	, CurveModificationDetector(CVarUseNewUndoSystem.GetValueOnAnyThread()
+		? MakePimpl<FModifyCurveObjectDetector>(
+			WeakCurveEditor, static_cast<EModifyDetectionLevel>(CVarModifyDetectionLevel.GetValueOnAnyThread()), ModifyDetectionWarning, ChangeDiff.GetSnapshot())
+		: nullptr
+		)
 {
 	const TSharedPtr<FCurveEditor> CurveEditorPin = WeakCurveEditor.Pin();
-	if (CurveEditorPin && !CVarUseNewUndoSystem.GetValueOnAnyThread())
+	if (!ensure(CurveEditorPin))
+	{
+		return;
+	}
+	
+	if (CVarUseNewUndoSystem.GetValueOnAnyThread())
+	{
+		// In the new undo system, we want to avoid the FCurveModel::GetOwningObject()->Modify() from being called.
+		// So we tell the curve models to invoke any PostChange-type of callbacks, which may cause Modify() calls.
+		// @see class doc string for why.
+		ForEachCurve(*CurveEditorPin, ChangeDiff, [](FCurveModel& CurveModel){ CurveModel.OpenChangeScope(); });
+	}
+	else
 	{
 		// If the new undo system is disabled, just do what we did before: call Modify().
-		ModifyCurves(CurveEditorPin, ChangeDiff);
+		ForEachCurve(*CurveEditorPin, ChangeDiff, [](FCurveModel& CurveModel){ CurveModel.Modify(); });
 	}
 }
 
@@ -68,16 +99,28 @@ FScopedCurveChange::FScopedCurveChange(
 
 FScopedCurveChange::~FScopedCurveChange()
 {
-	if (!CVarUseNewUndoSystem.GetValueOnAnyThread())
+	if (CVarUseNewUndoSystem.GetValueOnAnyThread())
 	{
-		return;
+		FinalizeChangesForNewUndoSystem();
 	}
+}
 	
+void FScopedCurveChange::FinalizeChangesForNewUndoSystem()
+{
 	const TSharedPtr<FCurveEditor> CurveEditorPin = WeakCurveEditor.Pin();
 	if (!CurveEditorPin)
 	{
 		return;
 	}
+	
+	// Match the constructor's OpenChangeScope call.
+	// CloseChangeScope may invoke some internal logic that causes Modify() calls on the FCurveModel::GetOwningObject.
+	// So it's important this call only takes places after we have appended our undo commands.
+	ON_SCOPE_EXIT
+	{
+		CurveModificationDetector.Reset(); // We've added the commands by now, if any. Stop warning about Modify() before CloseChangeScope.
+		ForEachCurve(*CurveEditorPin, ChangeDiff, [](FCurveModel& CurveModel){ CurveModel.CloseChangeScope(); });
+	};
 
 	const bool bRevertOnCancelled = EnumHasAnyFlags(Flags, EScopedKeyChangeFlags::RevertOnCancel);
 	const bool bHasBeenCancelled = IsCancelled();
@@ -99,23 +142,18 @@ FScopedCurveChange::~FScopedCurveChange()
 			*FText::AsMemory(InAllocationSize).ToString(),
 			FTimespan::FromSeconds(TimeMeasurement.GetTime()).GetTotalMilliseconds(),
 			*DumpSnapshotPerfData(ChangeDiff.GetSnapshot(), ECurvesSnapshotPerfFlags::All)
-			); 
+		); 
 	};
 	
 	if (!bHasBeenCancelled)
 	{
-		// The ProcessDiffs will call FTransactionManager::AppendChange repeatedly - once for each curve changed.
-		// As each call starts a FScopedTransaction, we should group them all under one transaction in case whatever code uses FScopedCurveChange
-		// has not started an outer FScopedTransaction of their own.
-		const FScopedTransaction GroupTransactions(Description);
-		
 		FMultiGenericCommandAppender CommandAppender(CurveEditorPin);
 		ChangeDiff.ProcessDiffs(CommandAppender);
 		LogPerfData(CommandAppender.GetAllocationSize()); 
 	}
 	else
 	{
-		FGenericCurveChangeData DeltaChange = ChangeDiff.ComputeDiff();
+		const FGenericCurveChangeData DeltaChange = ChangeDiff.ComputeDiff();
 		GenericCurveChange::RevertChange(*CurveEditorPin, DeltaChange);
 		LogPerfData(DeltaChange.GetAllocatedSize()); 
 	}
diff --git a/Engine/Source/Editor/CurveEditor/Public/CurveModel.h b/Engine/Source/Editor/CurveEditor/Public/CurveModel.h
--- a/Engine/Source/Editor/CurveEditor/Public/CurveModel.h
+++ b/Engine/Source/Editor/CurveEditor/Public/CurveModel.h
@@ -42,17 +42,16 @@ struct FKeyPosition;
 enum class ECurveEditorAxisOrientation : uint8;
 enum class ECurvePointType : uint8;
 
-namespace UE::CurveEditor 
-{ 
-	class ICurveEditorCurveCachePool;
-	struct FCurveDrawParamsHandle;
-}
+namespace UE::CurveEditor { class ICurveEditorCurveCachePool; }
+namespace UE::CurveEditor { struct FCurveDrawParamsHandle; }
+namespace UE::CurveEditor { class FScopedCurveChange; }
 
 /**
  * Class that models an underlying curve data structure through a generic abstraction that the curve editor understands.
  */
 class FCurveModel
 {
+	friend UE::CurveEditor::FScopedCurveChange; // For EnterChangeScope & ExitChangeScope.
 public:
 
 	FCurveModel()
@@ -62,7 +61,9 @@ public:
 	{}
 
 	virtual ~FCurveModel()
-	{}
+	{
+		ensureMsgf(!IsInChangeScope(), TEXT("Change scope was not closed"));
+	}
 
 	/**
 	 * Access the raw pointer of the curve data
@@ -689,6 +690,24 @@ protected:
 	/** Invokes by SetColor */
 	FSimpleMulticastDelegate PostColorChangedDelegate;
 
+	/** @return Whether this curve model is currently in a change scope. */
+	bool IsInChangeScope() const { return ChangeScopeCounter > 0; }
+
+	/**
+	 * Starts a scope of related changes.
+	 *
+	 * The curve model will receive a matching OnCloseRootChangeScope call when the scope ends.
+	 * The curve model implementation should use this to defer any post-change type of delegates until OnCloseRootChangeScope is invoked.
+	 *
+	 * The primary use case is recording changes using FScopedCurveChange. Do not call Modify() on the FCurveObject::GetOwningObject().
+	 * In the future, this may be extended so the curve model can append internal FChange commands as well, which are to be ordered correctly with
+	 * the changes made to the curve model. An example may be that SetKeyPositions should cause a call to UMovieSceneSection::ExpandToFrame.
+	 * @see FScopedCurveChange "Interaction with Modify()-based undo" section
+	 */
+	virtual void OnOpenRootChangeScope() {}
+	/** Closes the scope of related changes. The curve model can now call any callbacks it had deferred during the change scope. */
+	virtual void OnCloseRootChangeScope() {}
+	
 private:
 
 	/**
@@ -697,6 +716,32 @@ private:
 	 * The ID can only be set once during the lifetime of the FCurveModel to inconsistencies, e.g. TMap<FCurveModelId, TUniquePtr<FCurveModel>>.
 	 */
 	TOptional<FCurveModelID> CurveId;
+
+	/**
+	 * When positive, we're in a change scope. Otherwise, in no change scope.
+	 *
+	 * It is valid to mutate the curve model outside a change scope. The changes are just not considered to be related.
+	 */
+	int32 ChangeScopeCounter = 0;
+
+	/**
+	 * Starts a scope of related changes.
+	 * This is used by the new undo system to avoid Modify call on the FCurveModel::GetOwningObject().
+	 *
+	 * Every call must be matched by CloseChangeScope.
+	 * It's ok to call this multiple times after the scope has been started already.
+	 * @see OnOpenRootChangeScope, which is called when the scope is initially started.
+	 */
+	void OpenChangeScope();
+	/**
+	 * Closes a scope of related changes.
+	 * This is used by the new undo system.
+	 *
+	 * Every call matches a previous call to OpenChangeScope.
+	 * It's ok to call this multiple times after the scope has been started already.
+	 * @see OnCloseRootChangeScope, which is called when the scope outermost scope is closed.
+	 */
+	void CloseChangeScope();
 };
 
 #undef UE_API
diff --git a/Engine/Source/Editor/CurveEditor/Public/Modification/Utils/ScopedCurveChange.h b/Engine/Source/Editor/CurveEditor/Public/Modification/Utils/ScopedCurveChange.h
--- a/Engine/Source/Editor/CurveEditor/Public/Modification/Utils/ScopedCurveChange.h
+++ b/Engine/Source/Editor/CurveEditor/Public/Modification/Utils/ScopedCurveChange.h
@@ -24,7 +24,7 @@ enum class EScopedKeyChangeFlags : uint8
 };
 ENUM_CLASS_FLAGS(EScopedKeyChangeFlags);
 
-enum class ETrackedCurves : uint8 { All, Selected };
+class FModifyCurveObjectDetector;
 	
 /**
  * Starts a transaction. Detects changes made to the curves specified on construction and adds it as an undo-able action.
@@ -69,6 +69,40 @@ enum class ETrackedCurves : uint8 { All, Selected };
  * If you changed the declaration order, on redo, we'd first add to selection and then add the key to the curve.
  * Redo would no longer work: as FSelectionChangeCommand would add the FKeyHandle to the selection first, the selection code would defensively check
  * that a key exists first and reject the function call as the key has not yet been added to the curve model.
+ *
+ * ========== Interaction with Modify()-based undo ==========
+ * Summary: Do not call Modify() on UObjects that own a FCurveModel being tracked by FScopedCurveChange (@see FCurveModel::GetOwningObject).
+ * @see UE-331831 for fixing this issue
+ * Calling Modify() while FScopedCurveChange is in progress can cause subtle issues with the order of operations during undo / redo.
+ * 
+ * Example FChannelCurveModel:
+ * {
+ *		FCurveModel* FooModel = ...;
+ *		FCurveModel* BarModel = ...;
+ *		check(FooModel->GetOwningObject() == BarModel->GetOwningObject()); // The same object contains the curves.
+ *		FScopedCurveChange Change(...);		// Track changs made to FooModel
+ *		FooModel->AddKeys(...);				// Suppose this 1. mutates UObject to add the keys, then 2. calls Modify()
+ *		BarModel->AddKeys(...);				// Suppose this 1. mutates UObject to add the keys, then 2. calls Modify()
+ * }										// ~FScopedCurveChange now calls StoreUndo.
+ * 
+ * This scenario used to happen with MovieSceneTools' FChannelCurveModel. AddKeys would mutate the owning UMovieSceneParameterSection, so it contained the keys.
+ * At the end of the AddKeys call, the invocation of ISequencer::OnChannelChanged would cause ControlRig's auto-keying feature to call Modify on the
+ * owning object (in order to attempt to add a key for the key that was just added...).
+ * 
+ * This would result in the following undo buffer:
+ *   - FTransaction::Records[0] contains the delta changes from the first Modify() call. 
+ *   - FTransaction::Records[1] contains the command changes from ~FScopedCurveChange
+ *   
+ * Note, the first Modify would be called right after the keys were technically added to FooModel but not yet to BarModel.
+ * So the resulting base state for UMovieSceneParameterSection::BoolParameterNamesAndCurves would contain:
+ * - FooModel -> 1 key
+ * - BarModel -> 0 key
+ * The delta change would hence return reset BoolParameterNamesAndCurves so FooModel contains 1 key and BarModel contains 0 keys.
+ * 
+ * When you undo, the records are applied in reverse order:
+ * - FTransaction::Records[1] removes the keys added by AddKeys call.
+ * - FTransaction::Records[0] delta change is applied. Effectively, the key is added back to FooModel.
+ * @see UE-331831 for fixing this issue.
  */
 class FScopedCurveChange : public FScopedChangeBase
 {
@@ -102,7 +136,12 @@ private:
 	
 	/** Checks for changes. */
 	const FCurveChangeDiff ChangeDiff;
-
+	/** Flags about the change, e.g. whether to revert changes if cancelled, etc. */
 	const EScopedKeyChangeFlags Flags;
+
+	/** Util for warning about curves being modified. */
+	TPimplPtr<FModifyCurveObjectDetector> CurveModificationDetector;
+	
+	void FinalizeChangesForNewUndoSystem();
 };
 }
diff --git a/Engine/Source/Editor/MovieSceneTools/Private/Channels/ChannelCurveModel.cpp b/Engine/Source/Editor/MovieSceneTools/Private/Channels/ChannelCurveModel.cpp
--- a/Engine/Source/Editor/MovieSceneTools/Private/Channels/ChannelCurveModel.cpp
+++ b/Engine/Source/Editor/MovieSceneTools/Private/Channels/ChannelCurveModel.cpp
@@ -104,6 +104,16 @@ void FChannelCurveModel<ChannelType, ChannelValue, KeyType>::Modify()
 	LastSignature.Invalidate();
 }
 
+template <class ChannelType, class ChannelValue, class KeyType>
+void FChannelCurveModel<ChannelType, ChannelValue, KeyType>::OnCloseRootChangeScope()
+{
+	// The events are no longer deferred... invoke the delegates we deferred.
+	if (UMovieSceneSection* Section = WeakSection.Get())
+	{
+		InvokePostChannelChangedEvents(Section);
+	}
+}
+
 
 template <class ChannelType, class ChannelValue, class KeyType>
 void FChannelCurveModel<ChannelType, ChannelValue, KeyType>::DrawCurve(const FCurveEditor& CurveEditor, const FCurveEditorScreenSpace& ScreenSpace, TArray<TTuple<double, double>>& OutInterpolatingPoints) const
@@ -319,12 +329,14 @@ void FChannelCurveModel<ChannelType, ChannelValue, KeyType>::SetKeyPositions(TAr
 
 			if (!Section->GetRange().Contains(MinTime))
 			{
+				// TODO UE-305992 & UE-331831: This may cause Modify() to be called while FScopedCurveChange is active. That may cause ordering issues with undo / redo.
 				Section->ExpandToFrame(MinTime);
 				KeyOffset += MetaData->GetOffsetTime(Section) - Offset;
 			}
 
 			if (Min != Max && !Section->GetRange().Contains(MaxTime))
 			{
+				// TODO UE-305992 & UE-331831: This may cause Modify() to be called while FScopedCurveChange is active. That may cause ordering issues with undo / redo.
 				Section->ExpandToFrame(MaxTime);
 			}
 		}
@@ -363,7 +375,7 @@ void FChannelCurveModel<ChannelType, ChannelValue, KeyType>::SetKeyPositions(TAr
 			}
 		}
 		
-		OnPostChannelChanged(Channel, Section);
+		PostProcessChanges(Channel, Section);
 	}
 }
 
@@ -538,13 +550,14 @@ void FChannelCurveModel<ChannelType, ChannelValue, KeyType>::AddKeys(TArrayView<
 		if (InKeyPositions.Num() > 0)
 		{
 			FFrameNumber Offset = MetaData->GetOffsetTime(Section);
+			// TODO UE-305992 & UE-331831: This may cause Modify() to be called while FScopedCurveChange is active. That may cause ordering issues with undo / redo.
 			Section->ExpandToFrame(MinFrame + Offset);
 		}
 
 		// We reuse SetKeyAttributes here as there is complex logic determining which parts of the attributes are valid to set.
 		// For now we need to duplicate the new key handle array due to API mismatch. This will auto calculate tangents if needed.
 		SetKeyAttributes(NewKeyHandles, InKeyAttributes);
-		OnPostChannelChanged(Channel, Section);
+		PostProcessChanges(Channel, Section);
 	}
 }
 
@@ -574,7 +587,7 @@ void FChannelCurveModel<ChannelType, ChannelValue, KeyType>::RemoveKeys(TArrayVi
 			Channel->SetDefault(ValueAsCurveValue);
 		}
 
-		OnPostChannelChanged(Channel, Section);
+		PostProcessChanges(Channel, Section);
 	}
 }
 
@@ -607,7 +620,7 @@ void FChannelCurveModel<ChannelType, ChannelValue, KeyType>::ReplaceKeyHandles(
 			}
 
 			ChannelData.ReplaceKeyHandle(KeyIndex, InNewHandles[Index]);
-			OnPostChannelChanged(Channel, Section);
+			PostProcessChanges(Channel, Section);
 		}
 	}
 }
@@ -670,12 +683,27 @@ void FChannelCurveModel<ChannelType, ChannelValue, KeyType>::GetCurveColorObject
 }
 
 template <class ChannelType, class ChannelValue, class KeyType>
-void FChannelCurveModel<ChannelType, ChannelValue, KeyType>::OnPostChannelChanged(TNotNull<ChannelType*> Channel, TNotNull<UMovieSceneSection*> Section)
+void FChannelCurveModel<ChannelType, ChannelValue, KeyType>::PostProcessChanges(TNotNull<ChannelType*> Channel, TNotNull<UMovieSceneSection*> Section)
 {
 	LastSignature.Invalidate();
 	
 	// Allow channel to make any follow up changes to the keys before executing any other delegates
 	Channel->PostEditChange();
+
+	InvokePostChannelChangedEvents(Section);
+}
+
+template <class ChannelType, class ChannelValue, class KeyType>
+void FChannelCurveModel<ChannelType, ChannelValue, KeyType>::InvokePostChannelChangedEvents(TNotNull<UMovieSceneSection*> Section)
+{
+	// FScopeCurveChange starts a change scope and closes it again after it has appended its change commands.
+	// We need to avoid Modify() calls made to the section because it causes subtle ordering issues (@see FScopeCurveChange).
+	// OnChannelChanged may cause Modify() calls.
+	const bool bIsDeferringCallbacks = IsInChangeScope();
+	if (bIsDeferringCallbacks)
+	{
+		return;
+	}
 	
 	if (WeakSequencer.IsValid())
 	{
diff --git a/Engine/Source/Editor/MovieSceneTools/Public/Channels/ChannelCurveModel.h b/Engine/Source/Editor/MovieSceneTools/Public/Channels/ChannelCurveModel.h
--- a/Engine/Source/Editor/MovieSceneTools/Public/Channels/ChannelCurveModel.h
+++ b/Engine/Source/Editor/MovieSceneTools/Public/Channels/ChannelCurveModel.h
@@ -48,6 +48,7 @@ public:
 	virtual const void* GetCurve() const override;
 
 	virtual void Modify() override;
+	virtual void OnCloseRootChangeScope() override;
 
 	virtual void DrawCurve(const FCurveEditor& CurveEditor, const FCurveEditorScreenSpace& ScreenSpace, TArray<TTuple<double, double>>& OutInterpolatingPoints) const override;
 	virtual void GetKeys(double MinTime, double MaxTime, double MinValue, double MaxValue, TArray<FKeyHandle>& OutKeyHandles) const override;
@@ -106,17 +107,20 @@ protected:
 	virtual double GetKeyValue(TArrayView<const ChannelValue> Values, int32 Index) const = 0;
 	virtual void SetKeyValue(int32 Index, double KeyValue) const = 0;
 
+	/** Performs actions that should take place after any mutating call, like SetKeyPositions, etc. */
+	void PostProcessChanges(TNotNull<ChannelType*> Channel, TNotNull<UMovieSceneSection*> Section);
+	
 	/**
 	 * Updates all dependent systems after a change has been made to the channel, namely:
 	 *	- Invalidate LastSignature,
 	 *	- ChannelType::PostEditChange,
 	 *	- ChannelType::OnChannelChanged(),
 	 *	- UMovieSceneSection::MarkAsChanged,
-	 *	- MarkPackageDirty: this is required so the user is prompted to save the package after editing keys.
-	 *	If currently recording a transaction, the generated transaction will undo this flag when the user performs an undo op.
+	 *	- MarkPackageDirty: this is required so the user is prompted to save the package after editing keys. If currently recording a transaction,
+	 *	the generated transaction will undo this flag when the user performs an undo op.
 	 *	- FCurveModel::OnCurvesModified
 	 */
-	void OnPostChannelChanged(TNotNull<ChannelType*> Channel, TNotNull<UMovieSceneSection*> Section);
+	void InvokePostChannelChangedEvents(TNotNull<UMovieSceneSection*> Section);
 
 public:
 
