Curve Editor: Disable command-based undo system for 5.7. They will be reenabled for 5.8 to allow further testing.

diff --git a/Engine/Source/Editor/CurveEditor/Private/Modification/Utils/ScopedCurveChange.cpp b/Engine/Source/Editor/CurveEditor/Private/Modification/Utils/ScopedCurveChange.cpp
--- a/Engine/Source/Editor/CurveEditor/Private/Modification/Utils/ScopedCurveChange.cpp
+++ b/Engine/Source/Editor/CurveEditor/Private/Modification/Utils/ScopedCurveChange.cpp
@@ -27,11 +27,11 @@ static TAutoConsoleVariable<bool> CVarLogUndoKeyChangeSize(
 	);
 
 static TAutoConsoleVariable<bool> CVarUseNewUndoSystem(
-	TEXT("CurveEditor.UseNewUndoSystem"), true, TEXT("Whether to use the new undo system. Warning: There may still be some bugs.")
+	TEXT("CurveEditor.UseNewUndoSystem"), false, TEXT("Whether to use the new undo system. Warning: There may still be some bugs.")
 	);
 
 static TAutoConsoleVariable<int32> CVarModifyDetectionLevel(
-	TEXT("CurveEditor.UndoModifyDetectionLevel"), EModifyDetectionLevel::Ensure, TEXT("During FScopedCurveChange, detect Modify() calls made to UObjects that own FCurveModels. Modify is not safe to call during FScopedCurveChange.\n0 - No detection\n1 - log a warning\n2 - Ensure & log warning\n3 - EnsureAlways & log warning")
+	TEXT("CurveEditor.UndoModifyDetectionLevel"), EModifyDetectionLevel::None, TEXT("During FScopedCurveChange, detect Modify() calls made to UObjects that own FCurveModels. Modify is not safe to call during FScopedCurveChange.\n0 - No detection\n1 - log a warning\n2 - Ensure & log warning\n3 - EnsureAlways & log warning")
 );
 
 const TCHAR* ModifyDetectionWarning = TEXT(
