Curve Editor: Fix performance issue when no keys selected by caching common curve info instead of rebuilding it every frame.
- FSharedCurveInfoModel now holds the common curve data. It knows when that data needs refreshing.
- Promotes FCurveChangeListener from CurveEditorTools to CurveEditor
- Sequencer's FChannelCurveModel::HasChangedAndResetTest now broadcasts OnCurveModified when the signature has changed. This allows refreshing the cache when data changes externally, e.g. when you delete keys in Sequencer's key area.

diff --git a/Engine/Plugins/Editor/CurveEditorTools/Source/CurveEditorTools/Private/Misc/CurveChangeListener.cpp b/Engine/Plugins/Editor/CurveEditorTools/Source/CurveEditorTools/Private/Misc/CurveChangeListener.cpp
deleted file mode 100644
diff --git a/Engine/Plugins/Editor/CurveEditorTools/Source/CurveEditorTools/Private/Misc/CurveChangeListener.h b/Engine/Plugins/Editor/CurveEditorTools/Source/CurveEditorTools/Private/Misc/CurveChangeListener.h
deleted file mode 100644
diff --git a/Engine/Plugins/Editor/CurveEditorTools/Source/CurveEditorTools/Private/Tools/Lattice/CurveEditorLatticeTool.cpp b/Engine/Plugins/Editor/CurveEditorTools/Source/CurveEditorTools/Private/Tools/Lattice/CurveEditorLatticeTool.cpp
--- a/Engine/Plugins/Editor/CurveEditorTools/Source/CurveEditorTools/Private/Tools/Lattice/CurveEditorLatticeTool.cpp
+++ b/Engine/Plugins/Editor/CurveEditorTools/Source/CurveEditorTools/Private/Tools/Lattice/CurveEditorLatticeTool.cpp
@@ -297,7 +297,7 @@ FCurveEditorLatticeTool::FLatticeDeformerState::FLatticeDeformerState(const FLat
 		)
 	, GlobalDeformer(1, 1, Lattice.MinValues, Lattice.MaxValues)
 	, PerCurveData(BuildPerLatticeData(Lattice, *InCurveEditor))
-	, CurveChangeListener(InCurveEditor)
+	, CurveChangeListener(CurveEditor::FCurveChangeListener::MakeForAllCurves(InCurveEditor))
 	, PanelRebuildListener(InCurveEditor)
 {}
 	
diff --git a/Engine/Plugins/Editor/CurveEditorTools/Source/CurveEditorTools/Private/Tools/Lattice/CurveEditorLatticeTool.h b/Engine/Plugins/Editor/CurveEditorTools/Source/CurveEditorTools/Private/Tools/Lattice/CurveEditorLatticeTool.h
--- a/Engine/Plugins/Editor/CurveEditorTools/Source/CurveEditorTools/Private/Tools/Lattice/CurveEditorLatticeTool.h
+++ b/Engine/Plugins/Editor/CurveEditorTools/Source/CurveEditorTools/Private/Tools/Lattice/CurveEditorLatticeTool.h
@@ -115,7 +115,7 @@ private:
 		TMap<FCurveModelID, FPerCurveLatticeData> PerCurveData; 
 
 		/** Tells us when a curve is externally modified, so we can update the bounds. */
-		FCurveChangeListener CurveChangeListener;
+		CurveEditor::FCurveChangeListener CurveChangeListener;
 		/** Tells us when the panel is regenerated, in which case we must regenerate the lattice shape (e.g. could go from absolute -> normalized mode). */
 		FCurveViewRebuildListener PanelRebuildListener;
 		
diff --git a/Engine/Source/Editor/CurveEditor/Private/Misc/AttributeAccumulationUtils.cpp b/Engine/Source/Editor/CurveEditor/Private/Misc/AttributeAccumulationUtils.cpp
deleted file mode 100644
diff --git a/Engine/Source/Editor/CurveEditor/Private/Misc/AttributeAccumulationUtils.h b/Engine/Source/Editor/CurveEditor/Private/Misc/AttributeAccumulationUtils.h
deleted file mode 100644
diff --git a/Engine/Source/Editor/CurveEditor/Private/Misc/Attributes/AttributeAccumulationUtils.cpp b/Engine/Source/Editor/CurveEditor/Private/Misc/Attributes/AttributeAccumulationUtils.cpp
new file mode 100644
--- /dev/null
+++ b/Engine/Source/Editor/CurveEditor/Private/Misc/Attributes/AttributeAccumulationUtils.cpp
@@ -0,0 +1,98 @@
+// Copyright Epic Games, Inc. All Rights Reserved.
+
+#include "AttributeAccumulationUtils.h"
+
+#include "CurveEditor.h"
+#include "CurveEditorSelectionPrivate.h"
+#include "CurveEditorTrace.h"
+#include "CurveModel.h"
+
+namespace UE::CurveEditor
+{
+void UpdateCommonCurveInfo(
+	const FCurveEditor& InCurveEditor, bool& bOutSelectionSupportsWeightedTangents,
+	FCurveAttributes& OutCachedCommonCurveAttributes, FKeyAttributes& OutCachedCommonKeyAttributes)
+{
+	SCOPED_CURVE_EDITOR_TRACE(UpdateCommonCurveInfo);
+	
+	FAccumulatedCurveInfo AccumulatedCurveInfo = InCurveEditor.Selection.IsEmpty() ?
+		AccumulateAllCurvesInfo(InCurveEditor) :
+		AccumulateSelectedCurvesInfo(InCurveEditor);
+
+	// Reset the common curve and key info
+	bOutSelectionSupportsWeightedTangents = AccumulatedCurveInfo.bSelectionSupportsWeightedTangents;
+	OutCachedCommonCurveAttributes = AccumulatedCurveInfo.CurveAttributes.Get(FCurveAttributes());
+	OutCachedCommonKeyAttributes = AccumulatedCurveInfo.KeyAttributes.Get(FKeyAttributes());
+}
+
+FAccumulatedCurveInfo AccumulateSelectedCurvesInfo(const FCurveEditor& InCurveEditor)
+{
+	FAccumulatedCurveInfo AccumulatedAttributes;
+	for (const TTuple<FCurveModelID, FKeyHandleSet>& Pair : InCurveEditor.Selection.GetAll())
+	{
+		if (const FCurveModel* Curve = InCurveEditor.FindCurve(Pair.Key))
+		{
+			AccumulateCurveInfo(Curve, Pair.Value.AsArray(), AccumulatedAttributes);
+		}
+	}
+	return AccumulatedAttributes;
+}
+
+FAccumulatedCurveInfo AccumulateAllCurvesInfo(const FCurveEditor& InCurveEditor)
+{	
+	FAccumulatedCurveInfo AccumulatedAttributes;
+	for (const FCurveModelID ID : InCurveEditor.GetEditedCurves())
+	{
+		if (FCurveModel* Curve = InCurveEditor.FindCurve(ID))
+		{
+			const TArray<FKeyHandle> KeyHandles = Curve->GetAllKeys();
+			AccumulateCurveInfo(Curve, KeyHandles, AccumulatedAttributes);
+		}
+	}	
+	return AccumulatedAttributes;
+}
+
+void AccumulateCurveInfo(
+	const FCurveModel* InCurve, TConstArrayView<const FKeyHandle> InKeys, FAccumulatedCurveInfo& InOutAccumulatedAttributes)
+{
+	FCurveAttributes CurveAttributes;
+	InCurve->GetCurveAttributes(CurveAttributes);
+
+	// Some curves don't support extrapolation. We don't count them for determining the accumulated state.
+	const bool bCurveHasNoExtrapolation = CurveAttributes.HasPreExtrapolation() && CurveAttributes.GetPreExtrapolation() == RCCE_None
+		&& CurveAttributes.HasPostExtrapolation() && CurveAttributes.GetPostExtrapolation() == RCCE_None;
+	if (!bCurveHasNoExtrapolation && !InOutAccumulatedAttributes.CurveAttributes.IsSet())
+	{
+		InOutAccumulatedAttributes.CurveAttributes = CurveAttributes;
+	}
+	else if (!bCurveHasNoExtrapolation)
+	{
+		InOutAccumulatedAttributes.CurveAttributes = FCurveAttributes::MaskCommon(
+			InOutAccumulatedAttributes.CurveAttributes.GetValue(), CurveAttributes
+			);
+	}
+		
+	TArray<FKeyAttributes> AllKeyAttributes;	
+	AllKeyAttributes.SetNum(InKeys.Num());
+
+	InCurve->GetKeyAttributes(InKeys, AllKeyAttributes);
+	for (const FKeyAttributes& Attributes : AllKeyAttributes)
+	{
+		if (Attributes.HasTangentWeightMode())
+		{
+			InOutAccumulatedAttributes.bSelectionSupportsWeightedTangents = true;
+		}
+
+		if (!InOutAccumulatedAttributes.KeyAttributes.IsSet())
+		{
+			InOutAccumulatedAttributes.KeyAttributes = Attributes;
+		}
+		else
+		{
+			InOutAccumulatedAttributes.KeyAttributes = FKeyAttributes::MaskCommon(
+				InOutAccumulatedAttributes.KeyAttributes.GetValue(), Attributes
+				);
+		}
+	}
+}
+}
\ No newline at end of file
diff --git a/Engine/Source/Editor/CurveEditor/Private/Misc/Attributes/AttributeAccumulationUtils.h b/Engine/Source/Editor/CurveEditor/Private/Misc/Attributes/AttributeAccumulationUtils.h
new file mode 100644
--- /dev/null
+++ b/Engine/Source/Editor/CurveEditor/Private/Misc/Attributes/AttributeAccumulationUtils.h
@@ -0,0 +1,32 @@
+// Copyright Epic Games, Inc. All Rights Reserved.
+
+#pragma once
+
+#include "Curves/RealCurve.h"
+#include "CurveDataAbstraction.h"
+
+class FCurveEditor;
+class FCurveModel;
+
+namespace  UE::CurveEditor
+{
+	/** Used to hold curve and key attributes to update the Common Curve Info with. */
+	struct FAccumulatedCurveInfo
+	{
+		bool bSelectionSupportsWeightedTangents = false;
+		TOptional<FCurveAttributes> CurveAttributes;
+		TOptional<FKeyAttributes> KeyAttributes;
+	};
+
+	/** Update current curve selection attributes, if no keys are selected, info of all curves currently contained in the Curve Editor is used. */
+	void UpdateCommonCurveInfo(const FCurveEditor& InCurveEditor, bool& bOutSelectionSupportsWeightedTangents, FCurveAttributes& OutCachedCommonCurveAttributes, FKeyAttributes& OutCachedCommonKeyAttributes);
+
+	/** Gather curve and key attributes from all currently selected curves. */
+	FAccumulatedCurveInfo AccumulateSelectedCurvesInfo(const FCurveEditor& InCurveEditor);
+	/** Gather curve and key attributes from all curves currently contained in the Curve Editor regardless of visibility. .*/
+	FAccumulatedCurveInfo AccumulateAllCurvesInfo(const FCurveEditor& InCurveEditor);
+
+	/** Add curve and key attributes from the given curve and the keys on it to the accumulated info. */
+	void AccumulateCurveInfo(const FCurveModel* InCurve, TConstArrayView<const FKeyHandle> InKeys, FAccumulatedCurveInfo& InOutAccumulatedAttributes);
+}
+
diff --git a/Engine/Source/Editor/CurveEditor/Private/Misc/Attributes/SharedCurveInfoModel.cpp b/Engine/Source/Editor/CurveEditor/Private/Misc/Attributes/SharedCurveInfoModel.cpp
new file mode 100644
--- /dev/null
+++ b/Engine/Source/Editor/CurveEditor/Private/Misc/Attributes/SharedCurveInfoModel.cpp
@@ -0,0 +1,83 @@
+﻿// Copyright Epic Games, Inc. All Rights Reserved.
+
+#include "SharedCurveInfoModel.h"
+
+#include "AttributeAccumulationUtils.h"
+#include "CurveEditor.h"
+#include "Misc/CoreDelegates.h"
+
+namespace  UE::CurveEditor
+{
+FSharedCurveInfoModel::FSharedCurveInfoModel(const TSharedRef<FCurveEditor>& InCurveEditor)
+	: WeakCurveEditor(InCurveEditor.ToWeakPtr())
+	, CurveModifiedListener(FCurveChangeListener::MakeForAllCurves(InCurveEditor))
+{
+	InCurveEditor->Selection.OnSelectionChanged().AddRaw(this, &FSharedCurveInfoModel::DeferToRefreshToEndOfFrame);
+	InCurveEditor->OnCurveArrayChanged.AddRaw(this, &FSharedCurveInfoModel::OnCurveArrayChanged);
+	CurveModifiedListener.OnCurveModified().AddRaw(this, &FSharedCurveInfoModel::DeferToRefreshToEndOfFrame);
+
+	DeferToRefreshToEndOfFrame();
+}
+
+FSharedCurveInfoModel::~FSharedCurveInfoModel()
+{
+	if (const TSharedPtr<FCurveEditor> CurveEditorPin = WeakCurveEditor.Pin())
+	{
+		CurveEditorPin->Selection.OnSelectionChanged().RemoveAll(this);
+		CurveEditorPin->OnCurveArrayChanged.RemoveAll(this);
+	}
+	
+	ClearPendingRefresh();
+}
+
+void FSharedCurveInfoModel::Tick(float DeltaTime)
+{
+	const TSharedPtr<FCurveEditor> CurveEditorPin = WeakCurveEditor.Pin();
+	if (CurveEditorPin && CurveEditorPin->GetActiveCurvesSerialNumber() != ActiveCurvesSerialNumber)
+	{
+		DeferToRefreshToEndOfFrame();
+	}
+}
+
+TStatId FSharedCurveInfoModel::GetStatId() const
+{
+	RETURN_QUICK_DECLARE_CYCLE_STAT(FSharedCurveInfoModel, STATGROUP_Tickables); 
+}
+
+void FSharedCurveInfoModel::DeferToRefreshToEndOfFrame()
+{
+	if (!HasDeferredRefresh())
+	{
+		FCoreDelegates::OnEndFrame.AddRaw(this, &FSharedCurveInfoModel::Refresh);
+	}
+}
+
+void FSharedCurveInfoModel::Refresh()
+{
+	ClearPendingRefresh();
+
+	const TSharedPtr<FCurveEditor> CurveEditorPin = WeakCurveEditor.Pin();
+	if (!CurveEditorPin)
+	{
+		return;
+	}
+
+	TArray<FCurveModelID> Curves;
+	CurveEditorPin->GetCurves().GenerateKeyArray(Curves);
+	CurveModifiedListener.ResubscribeTo(MoveTemp(Curves));
+
+	ActiveCurvesSerialNumber = CurveEditorPin->GetActiveCurvesSerialNumber();
+	
+	UpdateCommonCurveInfo(*CurveEditorPin, bSelectionSupportsWeightedTangents, CachedCommonCurveAttributes, CachedCommonKeyAttributes);
+}
+
+bool FSharedCurveInfoModel::HasDeferredRefresh() const
+{
+	return FCoreDelegates::OnEndFrame.IsBoundToObject(this);
+}
+
+void FSharedCurveInfoModel::ClearPendingRefresh() const
+{
+	FCoreDelegates::OnEndFrame.RemoveAll(this);
+}
+}
diff --git a/Engine/Source/Editor/CurveEditor/Private/Misc/Attributes/SharedCurveInfoModel.h b/Engine/Source/Editor/CurveEditor/Private/Misc/Attributes/SharedCurveInfoModel.h
new file mode 100644
--- /dev/null
+++ b/Engine/Source/Editor/CurveEditor/Private/Misc/Attributes/SharedCurveInfoModel.h
@@ -0,0 +1,65 @@
+﻿// Copyright Epic Games, Inc. All Rights Reserved.
+
+#pragma once
+
+#include "CurveDataAbstraction.h"
+#include "CurveModel.h"
+#include "Misc/CurveChangeListener.h"
+#include "Templates/UnrealTemplate.h"
+#include "TickableEditorObject.h"
+
+class FCurveEditor;
+
+namespace  UE::CurveEditor
+{
+class FCurveDrawParamsCache;
+
+/**
+ * Keeps track of information that is shared for all keys: either for all selected keys, or if no keys are selected, all keys in view.
+ * Detects when changes are made to the curves and refreshes the informationa at the end of tick.
+ */
+class FSharedCurveInfoModel : public FNoncopyable, public FTickableEditorObject
+{
+public:
+	
+	explicit FSharedCurveInfoModel(const TSharedRef<FCurveEditor>& InCurveEditor);
+	~FSharedCurveInfoModel();
+
+	bool DoesSelectionSupportWeightedTangents() const { return bSelectionSupportsWeightedTangents; }
+	const FCurveAttributes& GetCachedCommonCurveAttributes() const { return CachedCommonCurveAttributes; }
+	const FKeyAttributes& GetCachedCommonKeyAttributes() const { return CachedCommonKeyAttributes; }
+
+	virtual void Tick(float DeltaTime) override;
+	virtual TStatId GetStatId() const override;
+
+private:
+
+	/** Owning curve editor. Needed to (un)subscribe from relevant change events. */
+	const TWeakPtr<FCurveEditor> WeakCurveEditor;
+
+	/** Subscribed to FCurveModel::OnCurveModified of all curves. */
+	FCurveChangeListener CurveModifiedListener;
+	/** Last FCurveEditor::ActiveCurvesSerialNumber. Used to detect changes. */
+	uint32 ActiveCurvesSerialNumber;
+	
+	/** True if the current selection supports weighted tangents, false otherwise */
+	bool bSelectionSupportsWeightedTangents = false;
+	/** Cached curve attributes that are common to all visible curves */
+	FCurveAttributes CachedCommonCurveAttributes;
+	/** Cached key attributes that are common to all selected keys */
+	FKeyAttributes CachedCommonKeyAttributes;
+
+	/** Enqueues a call to Refresh that will take place at the end of the current tick. */
+	void DeferToRefreshToEndOfFrame();
+	/** Updates the shared curve info */
+	void Refresh();
+	/** @return Whether a call to Refresh is enqueued */
+	bool HasDeferredRefresh() const;
+	/** Clears the subscription to FCoreDelegates. */
+	void ClearPendingRefresh() const;
+
+	/** On next Refresh, we'll resubscribe to all FCurveModel::OnCurveModified delegates. */
+	void OnCurveArrayChanged(FCurveModel*, bool, const FCurveEditor*) { DeferToRefreshToEndOfFrame(); }
+};
+}
+
diff --git a/Engine/Source/Editor/CurveEditor/Private/Misc/CurveChangeListener.cpp b/Engine/Source/Editor/CurveEditor/Private/Misc/CurveChangeListener.cpp
new file mode 100644
--- /dev/null
+++ b/Engine/Source/Editor/CurveEditor/Private/Misc/CurveChangeListener.cpp
@@ -0,0 +1,73 @@
+﻿// Copyright Epic Games, Inc. All Rights Reserved.
+
+#include "Misc/CurveChangeListener.h"
+
+#include "CurveEditor.h"
+#include "CurveModel.h"
+
+namespace UE::CurveEditor
+{
+FCurveChangeListener::FCurveChangeListener(const TSharedRef<FCurveEditor>& InCurveEditor, TArray<FCurveModelID> InCurvesToListenTo)
+	: WeakCurveEditor(InCurveEditor)
+	, SubscribedToCurves(MoveTemp(InCurvesToListenTo))
+{
+	SubscribeToCurves();
+}
+
+FCurveChangeListener::~FCurveChangeListener()
+{
+	UnsubscribeFromCurves();
+}
+
+void FCurveChangeListener::ResubscribeTo(TArray<FCurveModelID> InNewCurvesToSubscribeTo)
+{
+	UnsubscribeFromCurves();
+	SubscribedToCurves = MoveTemp(InNewCurvesToSubscribeTo);
+	SubscribeToCurves();
+}
+
+FCurveChangeListener FCurveChangeListener::MakeForAllCurves(const TSharedRef<FCurveEditor>& InCurveEditor)
+{
+	TArray<FCurveModelID> Curves;
+	InCurveEditor->GetCurves().GenerateKeyArray(Curves);
+	return FCurveChangeListener(InCurveEditor, MoveTemp(Curves));
+}
+
+FCurveChangeListener FCurveChangeListener::MakeForSelectedCurves(const TSharedRef<FCurveEditor>& InCurveEditor)
+{
+	TArray<FCurveModelID> Curves;
+	InCurveEditor->Selection.GetAll().GetKeys(Curves);
+	return FCurveChangeListener(InCurveEditor, MoveTemp(Curves));
+}
+
+void FCurveChangeListener::SubscribeToCurves()
+{
+	const TSharedPtr<FCurveEditor> CurveEditorPin = WeakCurveEditor.Pin();
+	if (!CurveEditorPin)
+	{
+		return;
+	}
+	
+	for (const FCurveModelID& CurveId : SubscribedToCurves)
+	{
+		if (FCurveModel* Model = CurveEditorPin->FindCurve(CurveId))
+		{
+			Model->OnCurveModified().AddRaw(this, &FCurveChangeListener::HandleCurveModified);
+		}
+	}
+}
+
+void FCurveChangeListener::UnsubscribeFromCurves()
+{
+	if (TSharedPtr<FCurveEditor> CurveEditorPin = WeakCurveEditor.Pin())
+	{
+		for (const FCurveModelID& CurveId : SubscribedToCurves)
+		{
+			if (FCurveModel* Model = CurveEditorPin->FindCurve(CurveId))
+			{
+				Model->OnCurveModified().RemoveAll(this);
+			}
+		}
+	}
+}
+}
diff --git a/Engine/Source/Editor/CurveEditor/Private/SCurveEditorPanel.cpp b/Engine/Source/Editor/CurveEditor/Private/SCurveEditorPanel.cpp
--- a/Engine/Source/Editor/CurveEditor/Private/SCurveEditorPanel.cpp
+++ b/Engine/Source/Editor/CurveEditor/Private/SCurveEditorPanel.cpp
@@ -67,8 +67,9 @@
 #include "Algo/AnyOf.h"
 #include "Filters/CurveEditorEulerFilter.h"
 #include "Filters/PromotedFilterContainer.h"
-#include "Misc/AttributeAccumulationUtils.h"
 #include "Misc/ResizeParamUtils.h"
+#include "Misc/Attributes/AttributeAccumulationUtils.h"
+#include "Misc/Attributes/SharedCurveInfoModel.h"
 #include "Modification/Utils/ScopedCurveChange.h"
 #include "Modification/Utils/ScopedSelectionChange.h"
 #include "Styling/AppStyle.h"
@@ -137,7 +138,6 @@ SCurveEditorPanel::SCurveEditorPanel()
 	, CachedActiveCurvesSerialNumber(-1)
 {
 	EditObjects = MakeUnique<FCurveEditorEditObjectContainer>();
-	bSelectionSupportsWeightedTangents = false;
 }
 
 SCurveEditorPanel::~SCurveEditorPanel()
@@ -151,6 +151,7 @@ void SCurveEditorPanel::Construct(const FArguments& InArgs, TSharedRef<FCurveEdi
 	GridLineTintAttribute = InArgs._GridLineTint;
 	DisabledTimeSnapTooltipAttribute = InArgs._DisabledTimeSnapTooltip;
 	WeakTabManager = InArgs._TabManager;
+	SharedCurveInfo = MakePimpl<UE::CurveEditor::FSharedCurveInfoModel>(InCurveEditor);
 
 	CachedSelectionSerialNumber = 0;
 
@@ -529,7 +530,6 @@ void SCurveEditorPanel::Tick(const FGeometry& AllottedGeometry, const double InC
 		CachedActiveCurvesSerialNumber = CurveEditor->GetActiveCurvesSerialNumber();
 	}
 
-	UE::CurveEditor::UpdateCommonCurveInfo(*CurveEditor, bSelectionSupportsWeightedTangents, CachedCommonCurveAttributes, CachedCommonKeyAttributes);
 	UpdateEditBox();
 	UpdateTime();
 
@@ -1033,6 +1033,36 @@ void SCurveEditorPanel::SetCurveAttributes(FCurveAttributes CurveAttributes, FTe
 	}
 }
 
+bool SCurveEditorPanel::CompareCommonInterpolationMode(ERichCurveInterpMode InterpMode) const
+{
+	const FKeyAttributes& CachedCommonKeyAttributes = SharedCurveInfo->GetCachedCommonKeyAttributes();
+	return CachedCommonKeyAttributes.HasInterpMode() && CachedCommonKeyAttributes.GetInterpMode() == InterpMode;
+}
+
+bool SCurveEditorPanel::CompareCommonTangentMode(ERichCurveInterpMode InterpMode, ERichCurveTangentMode TangentMode) const
+{
+	const FKeyAttributes& CachedCommonKeyAttributes = SharedCurveInfo->GetCachedCommonKeyAttributes();
+	return CompareCommonInterpolationMode(InterpMode) && CachedCommonKeyAttributes.HasTangentMode() && CachedCommonKeyAttributes.GetTangentMode() == TangentMode;
+}
+
+bool SCurveEditorPanel::CompareCommonTangentWeightMode(ERichCurveInterpMode InterpMode, ERichCurveTangentWeightMode TangentWeightMode) const
+{
+	const FKeyAttributes& CachedCommonKeyAttributes = SharedCurveInfo->GetCachedCommonKeyAttributes();
+	return CompareCommonInterpolationMode(InterpMode) && CachedCommonKeyAttributes.HasTangentWeightMode() && CachedCommonKeyAttributes.GetTangentWeightMode() == TangentWeightMode;
+}
+
+bool SCurveEditorPanel::CompareCommonPreExtrapolationMode(ERichCurveExtrapolation PreExtrapolationMode) const
+{
+	const FCurveAttributes& CachedCommonCurveAttributes = SharedCurveInfo->GetCachedCommonCurveAttributes();
+	return CachedCommonCurveAttributes.HasPreExtrapolation()  && CachedCommonCurveAttributes.GetPreExtrapolation() == PreExtrapolationMode;
+}
+
+bool SCurveEditorPanel::CompareCommonPostExtrapolationMode(ERichCurveExtrapolation PostExtrapolationMode) const
+{
+	const FCurveAttributes& CachedCommonCurveAttributes = SharedCurveInfo->GetCachedCommonCurveAttributes();
+	return CachedCommonCurveAttributes.HasPostExtrapolation() && CachedCommonCurveAttributes.GetPostExtrapolation() == PostExtrapolationMode;
+}
+
 void SCurveEditorPanel::ToggleWeightedTangents()
 {
 	using namespace UE::CurveEditor;
@@ -1089,7 +1119,7 @@ void SCurveEditorPanel::ToggleWeightedTangents()
 
 bool SCurveEditorPanel::CanToggleWeightedTangents() const
 {
-	return bSelectionSupportsWeightedTangents && CanSetKeyInterpolation();
+	return SharedCurveInfo->DoesSelectionSupportWeightedTangents() && CanSetKeyInterpolation();
 }
 
 bool SCurveEditorPanel::CanSetKeyInterpolation() const
diff --git a/Engine/Source/Editor/CurveEditor/Public/Misc/CurveChangeListener.h b/Engine/Source/Editor/CurveEditor/Public/Misc/CurveChangeListener.h
new file mode 100644
--- /dev/null
+++ b/Engine/Source/Editor/CurveEditor/Public/Misc/CurveChangeListener.h
@@ -0,0 +1,52 @@
+﻿// Copyright Epic Games, Inc. All Rights Reserved.
+
+#pragma once
+
+#include "Containers/Array.h"
+#include "Delegates/Delegate.h"
+#include "Templates/SharedPointer.h"
+#include "Templates/UnrealTemplate.h"
+
+#define UE_API CURVEEDITOR_API
+
+class FCurveEditor;
+struct FCurveModelID;
+
+namespace UE::CurveEditor
+{
+/** Listens to changes made to curves and cleans the subscriptions on destruction. */
+class FCurveChangeListener : public FNoncopyable
+{
+public:
+	
+	UE_API explicit FCurveChangeListener(const TSharedRef<FCurveEditor>& InCurveEditor, TArray<FCurveModelID> InCurvesToListenTo);
+	UE_API ~FCurveChangeListener();
+
+	/** Changes the curves to subscribe to. */
+	UE_API void ResubscribeTo(TArray<FCurveModelID> InNewCurvesToSubscribeTo);
+	
+	/** @return Change listener that subscribes to all FCurveModels on InCurveEditor.*/
+	UE_API static FCurveChangeListener MakeForAllCurves(const TSharedRef<FCurveEditor>& InCurveEditor);
+	/** @return Change listener that subscribed only to the FCurveModels CURRENTLY selected. */
+	UE_API static FCurveChangeListener MakeForSelectedCurves(const TSharedRef<FCurveEditor>& InCurveEditor);
+
+	FSimpleMulticastDelegate& OnCurveModified() { return OnCurveModifiedDelegate; }
+
+private:
+
+	/** Used to unsubscribe on destruction. */
+	TWeakPtr<FCurveEditor> WeakCurveEditor;
+	/** The curves we subscribed to */
+	TArray<FCurveModelID> SubscribedToCurves;
+
+	/** Invokes when any of the listeners is changed. */
+	FSimpleMulticastDelegate OnCurveModifiedDelegate;
+	
+	void HandleCurveModified() const { OnCurveModifiedDelegate.Broadcast(); }
+
+	void SubscribeToCurves();
+	void UnsubscribeFromCurves();
+};
+}
+
+#undef UE_API
\ No newline at end of file
diff --git a/Engine/Source/Editor/CurveEditor/Public/SCurveEditorPanel.h b/Engine/Source/Editor/CurveEditor/Public/SCurveEditorPanel.h
--- a/Engine/Source/Editor/CurveEditor/Public/SCurveEditorPanel.h
+++ b/Engine/Source/Editor/CurveEditor/Public/SCurveEditorPanel.h
@@ -53,6 +53,7 @@ struct FCurveEditorEditObjectContainer;
 struct FCurveEditorToolID;
 struct FKeyEvent;
 namespace UE::CurveEditor { class FPromotedFilterContainer; }
+namespace UE::CurveEditor { class FSharedCurveInfoModel; }
 
 /**
  * Curve editor widget that reflects the state of an FCurveEditor
@@ -288,34 +289,19 @@ private:
 	void SetCurveAttributes(FCurveAttributes CurveAttributes, FText Description);
 
 	/** Compare all the currently selected keys' interp modes against the specified interp mode */
-	bool CompareCommonInterpolationMode(ERichCurveInterpMode InterpMode) const
-	{
-		return CachedCommonKeyAttributes.HasInterpMode() && CachedCommonKeyAttributes.GetInterpMode() == InterpMode;
-	}
+	bool CompareCommonInterpolationMode(ERichCurveInterpMode InterpMode) const;
 
 	/** Compare all the currently selected keys' tangent modes against the specified tangent mode */
-	bool CompareCommonTangentMode(ERichCurveInterpMode InterpMode, ERichCurveTangentMode TangentMode) const
-	{
-		return CompareCommonInterpolationMode(InterpMode) && CachedCommonKeyAttributes.HasTangentMode() && CachedCommonKeyAttributes.GetTangentMode() == TangentMode;
-	}
+	bool CompareCommonTangentMode(ERichCurveInterpMode InterpMode, ERichCurveTangentMode TangentMode) const;
 
 	/** Compare all the currently selected keys' tangent modes against the specified tangent mode */
-	bool CompareCommonTangentWeightMode(ERichCurveInterpMode InterpMode, ERichCurveTangentWeightMode TangentWeightMode) const
-	{
-		return CompareCommonInterpolationMode(InterpMode) && CachedCommonKeyAttributes.HasTangentWeightMode() && CachedCommonKeyAttributes.GetTangentWeightMode() == TangentWeightMode;
-	}
+	bool CompareCommonTangentWeightMode(ERichCurveInterpMode InterpMode, ERichCurveTangentWeightMode TangentWeightMode) const;
 
 	/** Compare all the visible curves' pre-extrapolation modes against the specified extrapolation mode */
-	bool CompareCommonPreExtrapolationMode(ERichCurveExtrapolation PreExtrapolationMode) const
-	{
-		return CachedCommonCurveAttributes.HasPreExtrapolation()  && CachedCommonCurveAttributes.GetPreExtrapolation() == PreExtrapolationMode;
-	}
+	bool CompareCommonPreExtrapolationMode(ERichCurveExtrapolation PreExtrapolationMode) const;
 
 	/** Compare all the visible curves' post-extrapolation modes against the specified extrapolation mode */
-	bool CompareCommonPostExtrapolationMode(ERichCurveExtrapolation PostExtrapolationMode) const
-	{
-		return CachedCommonCurveAttributes.HasPostExtrapolation() && CachedCommonCurveAttributes.GetPostExtrapolation() == PostExtrapolationMode;
-	}
+	bool CompareCommonPostExtrapolationMode(ERichCurveExtrapolation PostExtrapolationMode) const;
 
 	/**
 	 * Toggle weighted tangents on the current selection
@@ -375,14 +361,8 @@ private:
 	/** (Optional) the current drag operation */
 	TOptional<FCurveEditorDelayedDrag> DragOperation;
 
-	/** Cached curve attributes that are common to all visible curves */
-	FCurveAttributes CachedCommonCurveAttributes;
-
-	/** Cached key attributes that are common to all selected keys */
-	FKeyAttributes CachedCommonKeyAttributes;
-
-	/** True if the current selection supports weighted tangents, false otherwise */
-	bool bSelectionSupportsWeightedTangents;
+	/** Keeps track about information that is equal for all selected keys, or all keys if none are selected. */
+	TPimplPtr<UE::CurveEditor::FSharedCurveInfoModel> SharedCurveInfo; 
 
 	/** Attribute used for retrieving the desired grid line color */
 	TAttribute<FLinearColor> GridLineTintAttribute;
diff --git a/Engine/Source/Editor/MovieSceneTools/Public/Channels/ChannelCurveModel.h b/Engine/Source/Editor/MovieSceneTools/Public/Channels/ChannelCurveModel.h
--- a/Engine/Source/Editor/MovieSceneTools/Public/Channels/ChannelCurveModel.h
+++ b/Engine/Source/Editor/MovieSceneTools/Public/Channels/ChannelCurveModel.h
@@ -93,6 +93,7 @@ public:
 			if (SignedOwner->GetSignature() != LastSignature)
 			{
 				LastSignature = SignedOwner->GetSignature();
+				CurveModifiedDelegate.Broadcast();
 				return true;
 			}
 			return false;
