Curve Editor: Fix side effects of broadcasting FCurveModel::OnCurveModified in FChannelCurveModel::HasChangedAndResetTest.

This change instead introduces a separate event. It is private to FCurveEditor for now because we may want to extract the curve draw cache to a dedicated model in 5.8; this avoids deprecation.
Having a dedicated event like this gets rid of the code debt introduced by 46514545. When working on a bug fix, I realized that the way I introduced the change, I made the system much more difficult to reason about.

One side effect is that the lattice tool refreshes its bounds after you let go of dragging it.

diff --git a/Engine/Source/Editor/CurveEditor/Private/CurveEditorCurveDrawParamsCache.cpp b/Engine/Source/Editor/CurveEditor/Private/CurveEditorCurveDrawParamsCache.cpp
--- a/Engine/Source/Editor/CurveEditor/Private/CurveEditorCurveDrawParamsCache.cpp
+++ b/Engine/Source/Editor/CurveEditor/Private/CurveEditorCurveDrawParamsCache.cpp
@@ -348,6 +348,8 @@ namespace UE::CurveEditor
 						CurveEditorView->bUpdateModelViewTransforms = true;
 						CurveEditorView->UpdateCustomAxes();
 					}
+					
+					OnCurveHasChangedExternallyDelegate.Broadcast(CurveModelID);
 				}
 			}
 
diff --git a/Engine/Source/Editor/CurveEditor/Private/Misc/Attributes/SharedCurveInfoModel.cpp b/Engine/Source/Editor/CurveEditor/Private/Misc/Attributes/SharedCurveInfoModel.cpp
--- a/Engine/Source/Editor/CurveEditor/Private/Misc/Attributes/SharedCurveInfoModel.cpp
+++ b/Engine/Source/Editor/CurveEditor/Private/Misc/Attributes/SharedCurveInfoModel.cpp
@@ -4,17 +4,20 @@
 
 #include "AttributeAccumulationUtils.h"
 #include "CurveEditor.h"
+#include "SCurveEditorViewContainer.h"
 #include "Misc/CoreDelegates.h"
 
 namespace  UE::CurveEditor
 {
-FSharedCurveInfoModel::FSharedCurveInfoModel(const TSharedRef<FCurveEditor>& InCurveEditor)
+FSharedCurveInfoModel::FSharedCurveInfoModel(const TSharedRef<FCurveEditor>& InCurveEditor, const TSharedRef<SCurveEditorViewContainer>& InViewContainer)
 	: WeakCurveEditor(InCurveEditor.ToWeakPtr())
+	, ViewContainer(InViewContainer.ToWeakPtr())
 	, CurveModifiedListener(FCurveChangeListener::MakeForAllCurves(InCurveEditor))
 {
 	InCurveEditor->Selection.OnSelectionChanged().AddRaw(this, &FSharedCurveInfoModel::DeferToRefreshToEndOfFrame);
 	InCurveEditor->OnCurveArrayChanged.AddRaw(this, &FSharedCurveInfoModel::OnCurveArrayChanged);
 	CurveModifiedListener.OnCurveModified().AddRaw(this, &FSharedCurveInfoModel::DeferToRefreshToEndOfFrame);
+	InViewContainer->OnCurveHasChangedExternally().AddRaw(this, &FSharedCurveInfoModel::OnCurveHasChanged);
 
 	DeferToRefreshToEndOfFrame();
 }
@@ -27,6 +30,11 @@ FSharedCurveInfoModel::~FSharedCurveInfoModel()
 		CurveEditorPin->OnCurveArrayChanged.RemoveAll(this);
 	}
 	
+	if (const TSharedPtr<SCurveEditorViewContainer> ViewContainerPin = ViewContainer.Pin())
+	{
+		ViewContainerPin->OnCurveHasChangedExternally().RemoveAll(this);
+	}
+		
 	ClearPendingRefresh();
 }
 
diff --git a/Engine/Source/Editor/CurveEditor/Private/Misc/Attributes/SharedCurveInfoModel.h b/Engine/Source/Editor/CurveEditor/Private/Misc/Attributes/SharedCurveInfoModel.h
--- a/Engine/Source/Editor/CurveEditor/Private/Misc/Attributes/SharedCurveInfoModel.h
+++ b/Engine/Source/Editor/CurveEditor/Private/Misc/Attributes/SharedCurveInfoModel.h
@@ -9,6 +9,7 @@
 #include "TickableEditorObject.h"
 
 class FCurveEditor;
+class SCurveEditorViewContainer;
 
 namespace  UE::CurveEditor
 {
@@ -22,7 +23,7 @@ class FSharedCurveInfoModel : public FNoncopyable, public FTickableEditorObject
 {
 public:
 	
-	explicit FSharedCurveInfoModel(const TSharedRef<FCurveEditor>& InCurveEditor);
+	explicit FSharedCurveInfoModel(const TSharedRef<FCurveEditor>& InCurveEditor, const TSharedRef<SCurveEditorViewContainer>& InViewContainer);
 	~FSharedCurveInfoModel();
 
 	bool DoesSelectionSupportWeightedTangents() const { return bSelectionSupportsWeightedTangents; }
@@ -37,6 +38,14 @@ private:
 	/** Owning curve editor. Needed to (un)subscribe from relevant change events. */
 	const TWeakPtr<FCurveEditor> WeakCurveEditor;
 
+	/**
+	 * This is used to detect any external changes made to curves.
+	 * Specifically, this UI calls FCurveModel::HasChangedAndResetTest and broadcasts a delegate when a change is detected.
+	 *
+	 * Ideally, we would not reference UI here but another model. 
+	 */
+	const TWeakPtr<SCurveEditorViewContainer> ViewContainer;
+
 	/** Subscribed to FCurveModel::OnCurveModified of all curves. */
 	FCurveChangeListener CurveModifiedListener;
 	/** Last FCurveEditor::ActiveCurvesSerialNumber. Used to detect changes. */
@@ -60,6 +69,8 @@ private:
 
 	/** On next Refresh, we'll resubscribe to all FCurveModel::OnCurveModified delegates. */
 	void OnCurveArrayChanged(FCurveModel*, bool, const FCurveEditor*) { DeferToRefreshToEndOfFrame(); }
+	/** Invoked when some FCurveModel::HasChangedAndRestTest returned true. */
+	void OnCurveHasChanged(const FCurveModelID&) { DeferToRefreshToEndOfFrame(); }
 };
 }
 
diff --git a/Engine/Source/Editor/CurveEditor/Private/SCurveEditorPanel.cpp b/Engine/Source/Editor/CurveEditor/Private/SCurveEditorPanel.cpp
--- a/Engine/Source/Editor/CurveEditor/Private/SCurveEditorPanel.cpp
+++ b/Engine/Source/Editor/CurveEditor/Private/SCurveEditorPanel.cpp
@@ -151,7 +151,6 @@ void SCurveEditorPanel::Construct(const FArguments& InArgs, TSharedRef<FCurveEdi
 	GridLineTintAttribute = InArgs._GridLineTint;
 	DisabledTimeSnapTooltipAttribute = InArgs._DisabledTimeSnapTooltip;
 	WeakTabManager = InArgs._TabManager;
-	SharedCurveInfo = MakePimpl<UE::CurveEditor::FSharedCurveInfoModel>(InCurveEditor);
 
 	CachedSelectionSerialNumber = 0;
 
@@ -287,7 +286,8 @@ void SCurveEditorPanel::Construct(const FArguments& InArgs, TSharedRef<FCurveEdi
 			]
 		];
 	
-
+	SharedCurveInfo = MakePimpl<UE::CurveEditor::FSharedCurveInfoModel>(InCurveEditor, CurveViewsContainer.ToSharedRef());
+	
 	if (InArgs._TreeContent.Widget != SNullWidget::NullWidget)
 	{
 		ChildSlot
diff --git a/Engine/Source/Editor/CurveEditor/Private/SCurveEditorView.cpp b/Engine/Source/Editor/CurveEditor/Private/SCurveEditorView.cpp
--- a/Engine/Source/Editor/CurveEditor/Private/SCurveEditorView.cpp
+++ b/Engine/Source/Editor/CurveEditor/Private/SCurveEditorView.cpp
@@ -470,6 +470,11 @@ TSharedPtr<FCurveEditor> SCurveEditorView::GetCurveEditor() const
 	return WeakCurveEditor.IsValid() ? WeakCurveEditor.Pin() : nullptr;
 }
 
+UE::CurveEditor::FOnCurveHasChangedExternally& SCurveEditorView::OnCurveHasChangedExternally()
+{
+	return CurveDrawParamsCache->OnCurveHasChangedExternallyDelegate;
+}
+
 void SCurveEditorView::SetOutputBounds(double InOutputMin, double InOutputMax, FCurveEditorViewAxisID AxisID)
 {
 	if (!bFixedOutputBounds && InOutputMin < InOutputMax)
diff --git a/Engine/Source/Editor/CurveEditor/Private/SCurveEditorViewContainer.cpp b/Engine/Source/Editor/CurveEditor/Private/SCurveEditorViewContainer.cpp
--- a/Engine/Source/Editor/CurveEditor/Private/SCurveEditorViewContainer.cpp
+++ b/Engine/Source/Editor/CurveEditor/Private/SCurveEditorViewContainer.cpp
@@ -189,6 +189,11 @@ bool SCurveEditorViewContainer::IsScrubTimeKeyEvent(const FKeyEvent& InKeyEvent)
 	return false;
 }
 
+void SCurveEditorViewContainer::BroadcastOnCurveHasChangedExternally(const FCurveModelID& InCurveModelId) const
+{
+	OnCurveHasChangedExternallyDelegate.Broadcast(InCurveModelId);
+}
+
 FReply SCurveEditorViewContainer::OnKeyDown(const FGeometry& MyGeometry, const FKeyEvent& InKeyEvent)
 {
 	FReply ToolReply = FReply::Unhandled();
@@ -517,10 +522,17 @@ void SCurveEditorViewContainer::AddView(TSharedRef<SCurveEditorView> ViewToAdd)
 	}
 
 	ViewToAdd->SetRetainerWidget(RetainerWidget);
+
+	ViewToAdd->OnCurveHasChangedExternally().AddRaw(this, &SCurveEditorViewContainer::BroadcastOnCurveHasChangedExternally);
 }
 
 void SCurveEditorViewContainer::Clear()
 {
+	for (const TSharedPtr<SCurveEditorView>& View : Views)
+	{
+		View->OnCurveHasChangedExternally().RemoveAll(this);
+	}
+	
 	Views.Reset();
 	ClearChildren();
 }
diff --git a/Engine/Source/Editor/CurveEditor/Private/SCurveEditorViewContainer.h b/Engine/Source/Editor/CurveEditor/Private/SCurveEditorViewContainer.h
--- a/Engine/Source/Editor/CurveEditor/Private/SCurveEditorViewContainer.h
+++ b/Engine/Source/Editor/CurveEditor/Private/SCurveEditorViewContainer.h
@@ -2,6 +2,7 @@
 
 #pragma once
 
+#include "CurveEditorCurveDrawParamsCache.h"
 #include "Containers/Array.h"
 #include "Containers/ArrayView.h"
 #include "HAL/Platform.h"
@@ -58,6 +59,9 @@ class SCurveEditorViewContainer : public SVerticalBox
 
 	UE_API void Clear();
 
+	/** Invoked when any of the views has invoked FCurveModel::HasChangedAndResetTest, and it returned true. */
+	UE::CurveEditor::FOnCurveHasChangedExternally& OnCurveHasChangedExternally() { return OnCurveHasChangedExternallyDelegate; }
+
 private:
 
 	// SWidget Interface
@@ -92,6 +96,8 @@ private:
 
 	UE_API bool IsScrubTimeKeyEvent(const FKeyEvent& InKeyEvent);
 
+	void BroadcastOnCurveHasChangedExternally(const FCurveModelID& InCurveModelId) const;
+
 private:
 
 	/** The curve editor pointer */
@@ -111,6 +117,9 @@ private:
 	/** Possible pointer to a retainer widget that we may need to force update*/
 	TSharedPtr<SRetainerWidget> RetainerWidget;
 
+	/** Invoked when any of the views has invoked FCurveModel::HasChangedAndResetTest, and it returned true. */
+	UE::CurveEditor::FOnCurveHasChangedExternally OnCurveHasChangedExternallyDelegate;
+
 	/** 
 	 * Whether or not this widget caught an OnMouseDown notification 
 	 * Used to check if the selection should be cleared
diff --git a/Engine/Source/Editor/CurveEditor/Public/CurveEditorCurveDrawParamsCache.h b/Engine/Source/Editor/CurveEditor/Public/CurveEditorCurveDrawParamsCache.h
--- a/Engine/Source/Editor/CurveEditor/Public/CurveEditorCurveDrawParamsCache.h
+++ b/Engine/Source/Editor/CurveEditor/Public/CurveEditorCurveDrawParamsCache.h
@@ -4,6 +4,7 @@
 
 #include "CurveDrawInfo.h"
 #include "CurveEditorSettings.h"
+#include "Delegates/Delegate.h"
 #include "Templates/SharedPointer.h"
 
 #define UE_API CURVEEDITOR_API
@@ -19,6 +20,9 @@ struct FCurveModelID;
 namespace UE::CurveEditor
 {		
 	struct FCurveDrawParamsHandle;
+	
+	/** Invoked when FCurveModel::HasChangedAndResetTest returns true. */
+	DECLARE_MULTICAST_DELEGATE_OneParam(FOnCurveHasChangedExternally, const FCurveModelID&);
 
 	/** Caches curve draw params for curves of a curve editor. */
 	class FCurveDrawParamsCache
@@ -113,6 +117,12 @@ namespace UE::CurveEditor
 
 		/** The curve editor that owns this cache */
 		TWeakPtr<FCurveEditor> WeakCurveEditor;
+
+		/**
+		 * Invoked when FCurveModel::HasChangedAndResetTest returns true.
+		 * Only used by SCurveEditorView at the moment thus not public API. 
+		 */
+		FOnCurveHasChangedExternally OnCurveHasChangedExternallyDelegate;
 	};
 }
 
diff --git a/Engine/Source/Editor/CurveEditor/Public/SCurveEditorView.h b/Engine/Source/Editor/CurveEditor/Public/SCurveEditorView.h
--- a/Engine/Source/Editor/CurveEditor/Public/SCurveEditorView.h
+++ b/Engine/Source/Editor/CurveEditor/Public/SCurveEditorView.h
@@ -19,6 +19,7 @@
 #include "Widgets/SCompoundWidget.h"
 #include "CurveEditorSettings.h"
 #include "CurveDrawInfo.h"
+#include "CurveEditorCurveDrawParamsCache.h"
 
 #define UE_API CURVEEDITOR_API
 
@@ -262,6 +263,9 @@ public:
 	/** Returns the cure editor associated with this view, or nullptr if the view is not specific to a curve editor */
 	UE_API TSharedPtr<FCurveEditor> GetCurveEditor() const;
 
+	/** Invoked when we detect that a curve needs to be redrawn, i.e. when FCurveModel::HasChangedAndResetTest returns true. */
+	UE::CurveEditor::FOnCurveHasChangedExternally& OnCurveHasChangedExternally();
+
 public:
 
 	/**
diff --git a/Engine/Source/Editor/MovieSceneTools/Public/Channels/ChannelCurveModel.h b/Engine/Source/Editor/MovieSceneTools/Public/Channels/ChannelCurveModel.h
--- a/Engine/Source/Editor/MovieSceneTools/Public/Channels/ChannelCurveModel.h
+++ b/Engine/Source/Editor/MovieSceneTools/Public/Channels/ChannelCurveModel.h
@@ -93,7 +93,6 @@ public:
 			if (SignedOwner->GetSignature() != LastSignature)
 			{
 				LastSignature = SignedOwner->GetSignature();
-				CurveModifiedDelegate.Broadcast();
 				return true;
 			}
 			return false;
