Curve Editor: Validate undo / redo command when applying them.

We check whether the FCurveModelID that the command is referencing still exists. If it does not, then probably there is a bug somewhere in the surrounding system. This helps us detect bugs. This can be disabled with the CurveEditor.ValidateCurveExistenceOnUndoRedo CVar.

Hitting the ensure usually means something is wrong with your system state. Your task is:
1. Check which is the FCurveModelID is referenced by InDeltaChange but not present in the FCurveEditor.
2. Then investigate what caused the FCurveModelID to no longer be present the FCurveEditor since the command was recorded.
2.1 A start idea: Maybe you performed some undos, which removed some FCurveModels, and then some redos, which did not add them back properly?
If after all of this you're of the opinion that it's correct system state that the FCurveModel does not exist, then we need to engineer a system for silencing this warning.

diff --git a/Engine/Source/Editor/CurveEditor/Private/Modification/Keys/GenericCurveChangeCommand.cpp b/Engine/Source/Editor/CurveEditor/Private/Modification/Keys/GenericCurveChangeCommand.cpp
--- a/Engine/Source/Editor/CurveEditor/Private/Modification/Keys/GenericCurveChangeCommand.cpp
+++ b/Engine/Source/Editor/CurveEditor/Private/Modification/Keys/GenericCurveChangeCommand.cpp
@@ -2,15 +2,55 @@
 
 #include "GenericCurveChangeCommand.h"
 
+#include "Algo/AllOf.h"
 #include "CurveEditor.h"
+#include "HAL/IConsoleManager.h"
 #include "Modification/Keys/GenericCurveChangeUtils.h"
+#include "SCurveEditor.h"
 
 namespace UE::CurveEditor
 {
+namespace Private
+{
+static TAutoConsoleVariable<bool> CVarValidateCurveExistenceOnUndoRedo(
+	TEXT("CurveEditor.ValidateCurveExistenceOnUndoRedo"),
+	true,
+	TEXT("Validates that the curves referenced by undo / redo commands exist and logs a failing ensure if not. This helps find bugs with undo / redo.")
+	);
+
+template<typename TValue>
+static bool AllCurvesExist(const TMap<FCurveModelID, TUniquePtr<FCurveModel>>& InCurves, const TMap<FCurveModelID, TValue>& InDataMap)
+{
+	return Algo::AllOf(InDataMap, [&InCurves](const TPair<FCurveModelID, TValue>& Pair)
+	{
+		return InCurves.Contains(Pair.Key);
+	});
+}
+
+static void ValidateCurveExistence(const TMap<FCurveModelID, TUniquePtr<FCurveModel>>& InCurves, const FGenericCurveChangeData& InDeltaChange)
+{
+	if (!CVarValidateCurveExistenceOnUndoRedo.GetValueOnAnyThread())
+	{
+		return;
+	}
+
+	const bool bAllExist = AllCurvesExist(InCurves, InDeltaChange.MoveKeysData.ChangedCurves)
+		&& AllCurvesExist(InCurves, InDeltaChange.AddKeysData.SavedCurveState)
+		&& AllCurvesExist(InCurves, InDeltaChange.RemoveKeysData.SavedCurveState)
+		&& AllCurvesExist(InCurves, InDeltaChange.KeyAttributeData.ChangedCurves)
+		&& AllCurvesExist(InCurves, InDeltaChange.CurveAttributeData.ChangeData);
+
+	// Hitting this means that the transaction contains data for FCurveModel that no longer is in the FCurveEditor. See CL description.
+	ensureMsgf(bAllExist, TEXT("The curve change references some curves that do not exist in FCurveEditor. Investigate."));
+	UE_CLOG(!bAllExist, LogCurveEditor, Warning, TEXT("The curve change references some curves that do not exist in FCurveEditor. Investigate."));
+}
+}
+	
 void FGenericCurveChangeCommand::Apply(UObject* Object)
 {
 	if (const TSharedPtr<FCurveEditor> CurveEditor = GetCurveEditor())
 	{
+		Private::ValidateCurveExistence(CurveEditor->GetCurves(), DeltaChange);
 		GenericCurveChange::ApplyChange(CurveEditor->GetCurves(), DeltaChange);
 	}
 }
@@ -19,6 +59,7 @@ void FGenericCurveChangeCommand::Revert(UObject* Object)
 {
 	if (const TSharedPtr<FCurveEditor> CurveEditor = GetCurveEditor())
 	{
+		Private::ValidateCurveExistence(CurveEditor->GetCurves(), DeltaChange);
 		GenericCurveChange::RevertChange(CurveEditor->GetCurves(), DeltaChange);
 	}
 }
