Curve Editor: Translating keys left and right using ctrl+ arrow left / right now supports undo / redo.

diff --git a/Engine/Source/Editor/CurveEditor/Private/CurveEditor.cpp b/Engine/Source/Editor/CurveEditor/Private/CurveEditor.cpp
--- a/Engine/Source/Editor/CurveEditor/Private/CurveEditor.cpp
+++ b/Engine/Source/Editor/CurveEditor/Private/CurveEditor.cpp
@@ -901,9 +901,12 @@ void FCurveEditor::ZoomToFitInternal(EAxisList::Type Axes, const TMap<FCurveMode
 
 void FCurveEditor::TranslateSelectedKeys(double SecondsToAdd)
 {
+	using namespace UE::CurveEditor;
 	if (Selection.Count() > 0)
 	{
-		for (const TTuple<FCurveModelID, FKeyHandleSet>& Pair : Selection.GetAll())
+		const TMap<FCurveModelID, FKeyHandleSet>& CurvesToKeys = Selection.GetAll();
+		const FScopedCurveChange CurveChange(FCurvesSnapshotBuilder(SharedThis(this), CurvesToKeys, ECurveChangeFlags::MoveKeys));
+		for (const TTuple<FCurveModelID, FKeyHandleSet>& Pair : CurvesToKeys)
 		{
 			if (FCurveModel* Curve = FindCurve(Pair.Key))
 			{
