Sequencer x Curve Editor: Fix Curve Editor undo no longer working if you undo past "Possess Object in Sequencer".

Curve editor will now automatically try to reuse FCurveModelID when the model is re-added to Sequencer. This is done by matching the metadata, like channel name, owner, etc. of the added curve model to any data that was in the FCurveEditor in the past. This requires that the metadata does not change inbetween adding back the curve models.
In the future, we may need to pass some comparison flags when the FCurveEditor is passed, e.g. to ignore comparing short display name for figuring out the FCurveModelID. For example, I'm not sure whether users can rename the short names manually within the Level Sequencer. Right now the comparison considers all of the FCurveModel string fields; the reason for that is that editors like the curve asset editor only set a short display name.

This removes the FCurveModelID tracking done in IKeyArea. The reason is because it introduces a bug: when you undo pass the "Possess Object in Sequencer" transaction, all IKeyAreas are regenerated and FCurveModel::CurveId will be set with a completely new ID even though the same curve is being added back. The new system in curve editor should be able to handle reusing the FCurveModelID from now on.

diff --git a/Engine/Source/Editor/CurveEditor/Private/CurveEditor.cpp b/Engine/Source/Editor/CurveEditor/Private/CurveEditor.cpp
--- a/Engine/Source/Editor/CurveEditor/Private/CurveEditor.cpp
+++ b/Engine/Source/Editor/CurveEditor/Private/CurveEditor.cpp
@@ -61,6 +61,7 @@
 #include "UObject/UnrealNames.h"
 #include "UnrealExporter.h"
 #include "Misc/SmartSnap.h"
+#include "Modification/Resolution/CurveModelIdCache.h"
 #include "Modification/Utils/ScopedCurveChange.h"
 #include "Modification/Utils/ScopedSelectionChange.h"
 #include "Selection/SelectionCleanser.h"
@@ -146,6 +147,8 @@ void FCurveEditor::InitCurveEditor(const FCurveEditorInitParams& InInitParams)
 	TransactionManager = MakeUnique<UE::CurveEditor::FTransactionManager>(SharedThis(this));
 
 	SelectionCleanser = MakePimpl<UE::CurveEditor::FSelectionCleanser>(SharedThis(this));
+
+	CurveModelIdCache = MakePimpl<UE::CurveEditor::FCurveModelIdCache>();
 }
 
 int32 FCurveEditor::GetSupportedTangentTypes()
@@ -224,11 +227,9 @@ FCurveModelID FCurveEditor::AddCurve(TUniquePtr<FCurveModel>&& InCurve)
 	check(InCurve);
 	
 	// The curve ID is relevant e.g. for undo / redo.
-	// You can undo / redo selecting keys: if undo past a transaction that called FCurveEditor::AddCurve, redoing that transaction needs to add back
-	// the same curve ID so redoing the key selection also works.
-	// If InCurve has no ID set, GetOrInitId will set it here. If the caller actually specifies an ID and that ID is added already, they have a
-	// mistake in their business logic.
-	const FCurveModelID& CurveId = InCurve->GetOrInitId();
+	// If InCurve has no ID set, InitCurveModelIdWithReusedOrNewId will set it here.
+	// If the caller actually specifies an ID and that ID is added already, they have a mistake in their business logic.
+	const FCurveModelID& CurveId = UE::CurveEditor::InitCurveModelIdWithReusedOrNewId(*CurveModelIdCache, *InCurve, CurveData);
 	if (!ensureMsgf(!CurveData.Contains(CurveId), TEXT("Investigate what caused the double-addition and fix it!")))
 	{
 		return FCurveModelID();
diff --git a/Engine/Source/Editor/CurveEditor/Private/Modification/Resolution/CurveModelIdCache.cpp b/Engine/Source/Editor/CurveEditor/Private/Modification/Resolution/CurveModelIdCache.cpp
new file mode 100644
--- /dev/null
+++ b/Engine/Source/Editor/CurveEditor/Private/Modification/Resolution/CurveModelIdCache.cpp
@@ -0,0 +1,115 @@
+﻿// Copyright Epic Games, Inc. All Rights Reserved.
+
+#include "CurveModelIdCache.h"
+
+#include "CurveEditorTrace.h"
+#include "CurveModel.h"
+
+namespace UE::CurveEditor
+{
+namespace IdCacheDetail
+{
+static TOptional<FCurveModelID> FindByOwner(const FCurveModelIdCache& InCache, const FCurveMetaDataIdentifiers& InMetaData)
+{
+	UObject* Owner = InMetaData.Owner.Get();
+	const TArray<FCurveModelID>* CurveIndices = InCache.OwnerToCurvesCache.Find(Owner);
+	if (!CurveIndices)
+	{
+		return {};
+	}
+
+	for (const FCurveModelID& CurveId : *CurveIndices)
+	{
+		const FCurveMetaDataIdentifiers* MetaData = FindMetaData(InCache, CurveId);
+		if (MetaData && InMetaData == *MetaData)
+		{
+			return CurveId;
+		}
+	}
+
+	return {};
+}
+
+static TOptional<FCurveModelID> FindByLinearSearch(const FCurveModelIdCache& InCache, const FCurveMetaDataIdentifiers& InMetaData)
+{
+	for (int32 Index = 0; Index < InCache.CurveMetaData.Num(); ++Index)
+	{
+		const FCurveModelIdCache::FCurveData& Data = InCache.CurveMetaData[Index];
+		if (Data.MetaData == InMetaData)
+		{
+			return Data.CurveId;
+		}
+	}
+	
+	return {};
+}
+
+static TOptional<FCurveModelID> FindCurveId(const FCurveModelIdCache& InCache, const FCurveMetaDataIdentifiers& InMetaData)
+{
+	if (const TOptional<FCurveModelID> OwnerSeachResult = FindByOwner(InCache, InMetaData))
+	{
+		return OwnerSeachResult;
+	}
+
+	if (const TOptional<FCurveModelID> SearchAllResult = FindByLinearSearch(InCache, InMetaData))
+	{
+		return SearchAllResult;
+	}
+
+	return {};
+}
+
+/** Inserts the curve data into presorted InCache maintaining the sort order (by FCurveModelID). */
+static int32 AddCurveModelId(FCurveModelIdCache& InCache, FCurveModelIdCache::FCurveData InData)
+{
+	const int32 InsertIndex = Algo::LowerBound(InCache.CurveMetaData, InData,
+		[](const FCurveModelIdCache::FCurveData& Left, const FCurveModelIdCache::FCurveData& Right)
+		{
+			return Left.CurveId < Right.CurveId;
+		});
+	InCache.CurveMetaData.Insert(MoveTemp(InData), InsertIndex);
+	return InsertIndex;
+}
+}
+	
+FCurveModelID InitCurveModelIdWithReusedOrNewId(
+	FCurveModelIdCache& InCache, FCurveModel& InCurveModel, const TMap<FCurveModelID, TUniquePtr<FCurveModel>>& InExistingCurves
+	)
+{
+	SCOPED_CURVE_EDITOR_TRACE(GetOrInitCurveModelId);
+	if (InCurveModel.GetId().IsSet())
+	{
+		return *InCurveModel.GetId();
+	}
+
+	FCurveMetaDataIdentifiers MetaData(InCurveModel);
+	const TOptional<FCurveModelID> ReusedId = IdCacheDetail::FindCurveId(InCache, MetaData);
+	if (ReusedId
+		// If this fails, adjust your curves to have unique identifier permutations (intention name, etc.), or call FCurveModel::InitCurveId before
+		// FCurveEditor::AddCurve. Undo / redo operations may not work properly until you fix it. See CL description.
+		&& ensure(!InExistingCurves.Contains(*ReusedId)))
+	{
+		InCurveModel.InitCurveId(*ReusedId);
+		return *ReusedId;
+	}
+	
+	UObject* Owner = MetaData.Owner.Get();
+	const FCurveModelID NewId = InCurveModel.GetOrInitId();
+	IdCacheDetail::AddCurveModelId(InCache, { NewId, MoveTemp(MetaData)});
+	if (Owner)
+	{
+		TArray<FCurveModelID>& Indices = InCache.OwnerToCurvesCache.FindOrAdd(Owner);
+		Indices.Add(NewId);
+	}
+	return NewId;
+}
+
+const FCurveMetaDataIdentifiers* FindMetaData(const FCurveModelIdCache& InCache, const FCurveModelID& InCurveId)
+{
+	const TArray<FCurveModelIdCache::FCurveData>& Data = InCache.CurveMetaData;
+	const int32 Index = Algo::BinarySearchBy(Data, InCurveId,
+		[](const FCurveModelIdCache::FCurveData& InData){ return InData.CurveId; }
+		);
+	return Data.IsValidIndex(Index) ? &Data[Index].MetaData : nullptr;
+}
+}
\ No newline at end of file
diff --git a/Engine/Source/Editor/CurveEditor/Private/Modification/Resolution/CurveModelIdCache.h b/Engine/Source/Editor/CurveEditor/Private/Modification/Resolution/CurveModelIdCache.h
new file mode 100644
--- /dev/null
+++ b/Engine/Source/Editor/CurveEditor/Private/Modification/Resolution/CurveModelIdCache.h
@@ -0,0 +1,76 @@
+﻿// Copyright Epic Games, Inc. All Rights Reserved.
+
+#pragma once
+
+#include "CurveEditorTypes.h"
+#include "Modification/Resolution/CurveMetaDataIdentifiers.h"
+#include "UObject/ObjectKey.h"
+
+class FCurveModel;
+
+namespace UE::CurveEditor
+{
+/**
+ * Keeps track of curve metadata so FCurveModelID are reused when conceptually the same FCurveModels are re-added to FCurveEditor.
+ * This is relevant for command-based undo / redo to work correctly.
+ *
+ * ====== Background ======
+ * Consider the following scenario
+ * 1. User clicks a curve to view it. AddCurve is called: a curve representing X may be added with ID=42,
+ * 2. User does an undo-able action. FGenericCurveChangeCommand saves change data mapped to ID=42.
+ * 3. User clicks another curve. RemoveCurve is called: the curve with ID=42 is removed.
+ * 4. User clicks the X curve again to view it again. AddCurve is called again: the curve representing X is adde back to FCurveEditor.
+ * In order for FGenericCurveChangeCommand to continue working, in step 4, the curve must be added back using the same FCurveModelID.
+ * 
+ * Systems can be explicit by calling FCurveModel::InitCurveId before calling FCurveEditor::AddCurve.
+ * However, there are 2 issues with this approach:
+ * 1. that is cumbersome to work with,
+ * 2. command-based undo was added later so many older editor systems currently do not bother to re-use FCurveModelIDs.
+ *
+ * ===== Class purpose ======
+ * So CurveEditor provides a "service" to such systems: provided that the FCurveModel they re-add has the same metadata (owner object, short name, etc.),
+ * CurveEditor will automatically attempt to re-use the ID that the re-added curve had when it was last added. That covers a lot of use cases.
+ * Systems for which this logic is not enough should just call FCurveModel::InitCurveId before calling FCurveEditor::AddCurve.
+ */
+struct FCurveModelIdCache
+{
+	/**
+	 * Maps UObjects to the FCurveModels they own or used to own.
+	 * Since most FCurveModel have an owning UObject, we can speed up search using this map.
+	 */
+	TMap<TObjectKey<UObject>, TArray<FCurveModelID>> OwnerToCurvesCache;
+	
+	struct FCurveData
+	{
+		FCurveModelID CurveId;
+		FCurveMetaDataIdentifiers MetaData;
+	};
+	/** Holds all curve data ever added to FCurveEditor. Never removed from. Sorted by FCurveModelID to allow for binary search. */
+	TArray<FCurveData> CurveMetaData;
+};
+
+/**
+ * Called by FCurveEditor::AddCurve to init the FCurveModelID for the new curve.
+ * Tries to reuse an old FCurveModelID if matches the metadata of an old curve.
+ * Priority order:
+ * - Use InCurveModel's FCurveModel::CurveId, if set. 
+ * - Search whether InCurveModel's metadata matches that of CurveMetaData in FCurveEditor in the past, then use the cached FCurveModelID.
+ * - Generate a new ID and cache.
+ *
+ * @param InCache Cache holding the metadata
+ * @param InCurveModel The curve about to be added
+ * @param InExistingCurves Curves that already exist. Used to handle clashing metadata (i.e. when InCurveModel's metadata already is in use)
+ */
+FCurveModelID InitCurveModelIdWithReusedOrNewId(
+	FCurveModelIdCache& InCache, FCurveModel& InCurveModel,
+	const TMap<FCurveModelID, TUniquePtr<FCurveModel>>& InExistingCurves
+	);
+
+/**
+ * Performs binary search to metadata associated with InCurveId.
+ * @return Metadata associated with InCurveId
+ */
+const FCurveMetaDataIdentifiers* FindMetaData(const FCurveModelIdCache& InCache, const FCurveModelID& InCurveId);
+}
+
+
diff --git a/Engine/Source/Editor/CurveEditor/Public/CurveEditor.h b/Engine/Source/Editor/CurveEditor/Public/CurveEditor.h
--- a/Engine/Source/Editor/CurveEditor/Public/CurveEditor.h
+++ b/Engine/Source/Editor/CurveEditor/Public/CurveEditor.h
@@ -58,6 +58,7 @@ struct FCurveEditorInitParams;
 struct FCurveEditorSnapMetrics;
 struct FGeometry;
 namespace UE::CurveEditor { class FSelectionCleanser; }
+namespace UE::CurveEditor { struct FCurveModelIdCache; }
 
 DECLARE_DELEGATE_OneParam(FOnSetBoolean, bool)
 DECLARE_MULTICAST_DELEGATE_OneParam(FOnActiveToolChanged, FCurveEditorToolID)
@@ -766,6 +767,12 @@ protected:
 	/** Makes sure that selection does not reference stale FKeyHandles or FCurveModelIds, i.e. after undo or an operation has removed keys. */
 	TPimplPtr<UE::CurveEditor::FSelectionCleanser> SelectionCleanser;
 
+	/**
+	 * Tries to re-use FCurveModelIDs when the conceptually same FCurveModel is added back to FCurveEditor using AddCurve.
+	 * @see FCurveModelIdCache.
+	 */
+	TPimplPtr<UE::CurveEditor::FCurveModelIdCache> CurveModelIdCache;
+
 private:
 
 	
diff --git a/Engine/Source/Editor/CurveEditor/Public/Modification/Resolution/CurveMetaDataIdentifiers.h b/Engine/Source/Editor/CurveEditor/Public/Modification/Resolution/CurveMetaDataIdentifiers.h
new file mode 100644
--- /dev/null
+++ b/Engine/Source/Editor/CurveEditor/Public/Modification/Resolution/CurveMetaDataIdentifiers.h
@@ -0,0 +1,59 @@
+﻿// Copyright Epic Games, Inc. All Rights Reserved.
+
+#pragma once
+
+#include "CurveModel.h"
+#include "Containers/UnrealString.h"
+#include "Internationalization/Text.h"
+#include "UObject/NameTypes.h"
+#include "UObject/WeakObjectPtr.h"
+#include "UObject/WeakObjectPtrTemplates.h"
+
+namespace UE::CurveEditor
+{
+/** Meta data about a FCurveModel useful to cache FCurveModelID when a curve with the same data is re-added to FCurveEditor. */
+struct FCurveMetaDataIdentifiers
+{
+	/** The UObject that owned the CurveModel. */
+	TWeakObjectPtr<UObject> Owner;
+	
+	/** This curve's short display name. Used in situations where other mechanisms provide enough context about what the curve is (such as "X") */
+	FText ShortDisplayName;
+	/** This curve's long display name. Used in situations where the UI doesn't provide enough context about what the curve is otherwise (such as "Floor.Transform.X") */
+	FText LongDisplayName;
+
+	/** This curve's short intention (such as Transform.X or Scale.X). Used internally to match up curves when saving/restoring curves between different objects. */
+	FString IntentionName;
+	/** 
+	 * This curve's long intention (such as foot_fk_l.Transform.X or foot_fk_r.Scale.X). Used internally to match up curves when saving/restoring curves between different objects.
+	 * Long intention names have priority in copy/paste over short intention names, but we fall back to short intention if it's unclear what the user is trying to do.
+	 */
+	FString LongIntentionName;
+
+	/** The original channel name, used mostly to make sure names match with BP/Scripting */
+	FName ChannelName;
+
+	explicit FCurveMetaDataIdentifiers(const FCurveModel& InCurveModel)
+		: Owner(InCurveModel.GetOwningObject())
+		, ShortDisplayName(InCurveModel.GetShortDisplayName())
+		, LongDisplayName(InCurveModel.GetLongDisplayName())
+		, IntentionName(InCurveModel.GetIntentionName())
+		, LongIntentionName(InCurveModel.GetLongIntentionName())
+		, ChannelName(InCurveModel.GetChannelName())
+	{}
+
+	friend bool operator==(const FCurveMetaDataIdentifiers& Left, const FCurveMetaDataIdentifiers& Right)
+	{
+		return Left.Owner == Right.Owner
+			&& Left.ChannelName == Right.ChannelName
+			&& Left.IntentionName == Right.IntentionName
+			&& Left.LongIntentionName == Right.LongIntentionName
+			&& Left.ShortDisplayName.EqualTo(Right.ShortDisplayName)
+			&& Left.LongDisplayName.EqualTo(Right.LongDisplayName);
+	}
+	friend bool operator!=(const FCurveMetaDataIdentifiers& Left, const FCurveMetaDataIdentifiers& Right)
+	{
+		return !(Left == Right);
+	}
+};
+}
diff --git a/Engine/Source/Editor/Sequencer/Private/IKeyArea.cpp b/Engine/Source/Editor/Sequencer/Private/IKeyArea.cpp
--- a/Engine/Source/Editor/Sequencer/Private/IKeyArea.cpp
+++ b/Engine/Source/Editor/Sequencer/Private/IKeyArea.cpp
@@ -17,10 +17,9 @@
 #include "CurveEditorSettings.h"
 #include "CurveEditorTypes.h"
 
-IKeyArea::IKeyArea(TWeakPtr<ISequencerSection> InSection, FMovieSceneChannelHandle InChannel , const FCurveModelID& InCurveId)
+IKeyArea::IKeyArea(TWeakPtr<ISequencerSection> InSection, FMovieSceneChannelHandle InChannel)
 	: TreeSerialNumber(0)
 	, ChannelHandle(InChannel)
-	, CurveModelID(MakePimpl<FCurveModelID>(InCurveId))
 {
 	Reinitialize(InSection, InChannel);
 }
@@ -489,13 +488,6 @@ TUniquePtr<FCurveModel> IKeyArea::CreateCurveEditorModel(TSharedRef<ISequencer>
 					CurveModel->SetColor(SettingColor.GetValue());
 				}
 			}
-
-			// Make sure the curve consistently has the same ID.
-			// If the curve is added, removed, then added again to curve editor, systems will recognize this curve as the "same".
-			// This is relevant e.g. for undo / redo.
-			// You can undo / redo selecting keys: if undo past a transaction that called FCurveEditor::AddCurve,
-			// redoing that transaction needs to add back the same curve ID so redoing the key selection also works.
-			CurveModel->InitCurveId(*CurveModelID);
 		}
 		return CurveModel;
 	}
diff --git a/Engine/Source/Editor/Sequencer/Private/MVVM/ViewModels/ChannelModel.cpp b/Engine/Source/Editor/Sequencer/Private/MVVM/ViewModels/ChannelModel.cpp
--- a/Engine/Source/Editor/Sequencer/Private/MVVM/ViewModels/ChannelModel.cpp
+++ b/Engine/Source/Editor/Sequencer/Private/MVVM/ViewModels/ChannelModel.cpp
@@ -143,8 +143,7 @@ UE_SEQUENCER_DEFINE_VIEW_MODEL_TYPE_ID(FDynamicChannelGroupMuteExtension);
 
 
 FChannelModel::FChannelModel(FName InChannelName, TWeakPtr<ISequencerSection> InSection, FMovieSceneChannelHandle InChannel)
-	: CurveModelID(MakePimpl<FCurveModelID>(FCurveModelID::Unique()))
-	, KeyArea(MakeShared<IKeyArea>(InSection, InChannel, *CurveModelID))
+	: KeyArea(MakeShared<IKeyArea>(InSection, InChannel))
 	, ChannelName(InChannelName)
 {
 }
@@ -167,7 +166,7 @@ void FChannelModel::Initialize(TWeakPtr<ISequencerSection> InSection, FMovieScen
 {
 	if (!KeyArea)
 	{
-		KeyArea = MakeShared<IKeyArea>(InSection, InChannel, *CurveModelID);
+		KeyArea = MakeShared<IKeyArea>(InSection, InChannel);
 	}
 	else
 	{
diff --git a/Engine/Source/Editor/Sequencer/Public/IKeyArea.h b/Engine/Source/Editor/Sequencer/Public/IKeyArea.h
--- a/Engine/Source/Editor/Sequencer/Public/IKeyArea.h
+++ b/Engine/Source/Editor/Sequencer/Public/IKeyArea.h
@@ -52,7 +52,7 @@ public:
 	 * @param InSection The section that owns the channel that this key area represents
 	 * @param InChannel Handle to the channel this key area represents
 	 */
-	UE_API IKeyArea(TWeakPtr<ISequencerSection> InSection, FMovieSceneChannelHandle InChannel, const FCurveModelID& InCurveId);
+	UE_API IKeyArea(TWeakPtr<ISequencerSection> InSection, FMovieSceneChannelHandle InChannel);
 
 	/**
 	 * Reinitialize this key area with a new channel
@@ -346,9 +346,6 @@ private:
 
 	/** The display text of this channel. */
 	FText DisplayText;
-
-	/** Consistent ID that is reused and injected into the curve model every time it is created. Pimpl to avoid public build.cs dependency. */
-	const TPimplPtr<FCurveModelID> CurveModelID;
 };
 
 #undef UE_API
diff --git a/Engine/Source/Editor/Sequencer/Public/MVVM/ViewModels/ChannelModel.h b/Engine/Source/Editor/Sequencer/Public/MVVM/ViewModels/ChannelModel.h
--- a/Engine/Source/Editor/Sequencer/Public/MVVM/ViewModels/ChannelModel.h
+++ b/Engine/Source/Editor/Sequencer/Public/MVVM/ViewModels/ChannelModel.h
@@ -117,9 +117,6 @@ private:
 	UE_API FLinearColor GetKeyBarColor() const;
 
 private:
-
-	/** Consistent ID that is reused and injected into the curve model every time it is created. Pimpl to avoid public build.cs dependency. */
-	const TPimplPtr<FCurveModelID> CurveModelID;
 	
 	TSharedPtr<IKeyArea> KeyArea;
 	const FName ChannelName;
