Curve Editor: FGenericCurveChangeCommand can now be applied to FCurveModels that no longer visible, i.e. that are no longer in FCurveEditor::CurveData.

In many editors, such as Niagara curve editor, the user can make a change to curve X, then view curve Y (which reviews FCurveModel of X from FCurveEditor::CurvedData), and then attempt to do an undo operation.
- Before this change, the result was a no-op, i.e. missing data.
- With this change, the data can now be applied.

The solution is to instantiate a temporary FCurveModel. This is done through a delegate passed via FCurveEditorInitParams::ResolveCurveModelDelegate. The editors, like Sequencer or Niagara, are supposed to implement a delegate that looks up the curve data based on metadata the curve had before it was removed (intention name, display name, owner object, etc.): the delegate then temporarily instantiates a FCurveModel.
The drawback of this approach is that every editor in the engine needs to implement this delegate.

diff --git a/Engine/Source/Editor/CurveEditor/Private/CurveEditor.cpp b/Engine/Source/Editor/CurveEditor/Private/CurveEditor.cpp
--- a/Engine/Source/Editor/CurveEditor/Private/CurveEditor.cpp
+++ b/Engine/Source/Editor/CurveEditor/Private/CurveEditor.cpp
@@ -62,6 +62,7 @@
 #include "UnrealExporter.h"
 #include "Misc/SmartSnap.h"
 #include "Modification/Resolution/CurveModelIdCache.h"
+#include "Modification/Resolution/CurveModelLookUpInfo.h"
 #include "Modification/Utils/ScopedCurveChange.h"
 #include "Modification/Utils/ScopedSelectionChange.h"
 #include "Selection/SelectionCleanser.h"
@@ -115,6 +116,7 @@ void FCurveEditor::InitCurveEditor(const FCurveEditorInitParams& InInitParams)
 
 	Selection = FCurveEditorSelection(SharedThis(this));
 	ZoomScalingAttr = InInitParams.ZoomScalingAttr;
+	ResolveCurveModelDelegate = InInitParams.ResolveCurveModelDelegate;
 
 	// Editor Extensions can be registered in the Curve Editor module. To allow users to derive from FCurveEditor
 	// we have to manually reach out to the module and get a list of extensions to create an instance of them.
@@ -2671,6 +2673,13 @@ void FCurveEditor::OnAxisSnappingChanged()
 	}
 }
 
+TUniquePtr<FCurveModel> FCurveEditor::ResolveCurve(const FCurveModelID& InCurveID) const
+{
+	return ResolveCurveModelDelegate.IsBound()
+	? ResolveCurveModelDelegate.Execute(UE::CurveEditor::FCurveModelLookUpArgs(*this, InCurveID))
+	: nullptr;
+}
+
 void FCurveEditor::HandleCurveColorChanged(FCurveModelID CurveId)
 {
 	if (IsBroadcasting())
diff --git a/Engine/Source/Editor/CurveEditor/Private/Modification/Keys/GenericCurveChangeCommand.cpp b/Engine/Source/Editor/CurveEditor/Private/Modification/Keys/GenericCurveChangeCommand.cpp
--- a/Engine/Source/Editor/CurveEditor/Private/Modification/Keys/GenericCurveChangeCommand.cpp
+++ b/Engine/Source/Editor/CurveEditor/Private/Modification/Keys/GenericCurveChangeCommand.cpp
@@ -7,6 +7,7 @@
 #include "HAL/IConsoleManager.h"
 #include "Modification/Keys/GenericCurveChangeUtils.h"
 #include "SCurveEditor.h"
+#include "Modification/Resolution/ResolveCurveModelPasskey.h"
 
 namespace UE::CurveEditor
 {
@@ -18,40 +19,104 @@ static TAutoConsoleVariable<bool> CVarValidateCurveExistenceOnUndoRedo(
 	TEXT("Validates that the curves referenced by undo / redo commands exist and logs a failing ensure if not. This helps find bugs with undo / redo.")
 	);
 
-template<typename TValue>
-static bool AllCurvesExist(const TMap<FCurveModelID, TUniquePtr<FCurveModel>>& InCurves, const TMap<FCurveModelID, TValue>& InDataMap)
+template<typename TPredicate>
+concept CDoesCurveExistCallable = std::is_invocable_r_v<bool, TPredicate, const FCurveModelID&>;
+
+/** @return Whether all curves in InDataMap exists according to InPredicate. */
+template<typename TValue, CDoesCurveExistCallable Predicate>
+static bool AllCurvesExist(const TMap<FCurveModelID, TValue>& InDataMap, Predicate&& InPredicate)
 {
-	return Algo::AllOf(InDataMap, [&InCurves](const TPair<FCurveModelID, TValue>& Pair)
+	return Algo::AllOf(InDataMap, [&InPredicate](const TPair<FCurveModelID, TValue>& Pair)
 	{
-		return InCurves.Contains(Pair.Key);
+		return InPredicate(Pair.Key);
 	});
 }
 
-static void ValidateCurveExistence(const TMap<FCurveModelID, TUniquePtr<FCurveModel>>& InCurves, const FGenericCurveChangeData& InDeltaChange)
+/** Validates that all FCurveModelIDs referenced by FGenericCurveChangeData have a matching FCurveModel. */
+template<CDoesCurveExistCallable Predicate>
+static void ValidateCurveExistence(const FGenericCurveChangeData& InDeltaChange, Predicate&& InPredicate)
 {
 	if (!CVarValidateCurveExistenceOnUndoRedo.GetValueOnAnyThread())
 	{
 		return;
 	}
 
-	const bool bAllExist = AllCurvesExist(InCurves, InDeltaChange.MoveKeysData.ChangedCurves)
-		&& AllCurvesExist(InCurves, InDeltaChange.AddKeysData.SavedCurveState)
-		&& AllCurvesExist(InCurves, InDeltaChange.RemoveKeysData.SavedCurveState)
-		&& AllCurvesExist(InCurves, InDeltaChange.KeyAttributeData.ChangedCurves)
-		&& AllCurvesExist(InCurves, InDeltaChange.CurveAttributeData.ChangeData);
+	const bool bAllExist = AllCurvesExist(InDeltaChange.MoveKeysData.ChangedCurves, InPredicate)
+		&& AllCurvesExist(InDeltaChange.AddKeysData.SavedCurveState, InPredicate)
+		&& AllCurvesExist(InDeltaChange.RemoveKeysData.SavedCurveState, InPredicate)
+		&& AllCurvesExist(InDeltaChange.KeyAttributeData.ChangedCurves, InPredicate)
+		&& AllCurvesExist(InDeltaChange.CurveAttributeData.ChangeData, InPredicate);
 
 	// Hitting this means that the transaction contains data for FCurveModel that no longer is in the FCurveEditor. See CL description.
+	// See also FCurveEditorInitParams::ResolveCurveModelDelegate.
 	ensureMsgf(bAllExist, TEXT("The curve change references some curves that do not exist in FCurveEditor. Investigate."));
 	UE_CLOG(!bAllExist, LogCurveEditor, Warning, TEXT("The curve change references some curves that do not exist in FCurveEditor. Investigate."));
 }
+
+/**
+ * Goes through undo data and finds FCurveModels that are not currently in FCurveEditor::CurveData.
+ * The missing FCurveModels are appended to OutInstantiatedCurves.
+ */
+template<typename TValue>
+static void InstantiateMissingCurves(
+	const FCurveEditor& InCurveEditor, const TMap<FCurveModelID, TValue>& InReferencingData,
+	TMap<FCurveModelID, TUniquePtr<FCurveModel>>& OutInstantiatedCurves
+	)
+{
+	for (const TPair<FCurveModelID, TValue>& DataPair : InReferencingData)
+	{
+		const FCurveModelID& CurveId = DataPair.Key;
+		if (InCurveEditor.GetCurves().Contains(CurveId)) // We're only interested in missing curves
+		{
+			continue;
+		}
+
+		if (TUniquePtr<FCurveModel> InstantiatedCurveModel = ResolveCurveModel(InCurveEditor, CurveId))
+		{
+			OutInstantiatedCurves.Emplace(CurveId, MoveTemp(InstantiatedCurveModel));
+		}
+	}
+}
+
+/**
+ * Invokes InProcessCurves once for already existing and missing curves, i.e.
+ * - the curves currently in FCurveEditor::CurveData (existing), and
+ * - a temporary TMap of FCurveModelIDs that were missing from FCurveEditor::CurveData but that were referenced by the undo data (missing)
+ */
+template<typename TProcessCurves> requires std::is_invocable_v<TProcessCurves, TMap<FCurveModelID, TUniquePtr<FCurveModel>>>
+void ProcessExistingAndMissingCurves(const FCurveEditor& InCurveEditor, const FGenericCurveChangeData& InDeltaChange, TProcessCurves&& InProcessCurves)
+{
+	const TMap<FCurveModelID, TUniquePtr<FCurveModel>>& RegisteredCurves = InCurveEditor.GetCurves();
+	TMap<FCurveModelID, TUniquePtr<FCurveModel>> TemporaryCurves;
+	{
+		InstantiateMissingCurves(InCurveEditor, InDeltaChange.MoveKeysData.ChangedCurves, TemporaryCurves);
+		InstantiateMissingCurves(InCurveEditor, InDeltaChange.AddKeysData.SavedCurveState, TemporaryCurves);
+		InstantiateMissingCurves(InCurveEditor, InDeltaChange.RemoveKeysData.SavedCurveState, TemporaryCurves);
+		InstantiateMissingCurves(InCurveEditor, InDeltaChange.KeyAttributeData.ChangedCurves, TemporaryCurves);
+		InstantiateMissingCurves(InCurveEditor, InDeltaChange.CurveAttributeData.ChangeData, TemporaryCurves);
+	}
+
+	// We'll log warnings if the undo data cannot be applied due to a missing FCurveModel. We do this so we can flag potential bugs to UE developers.
+	Private::ValidateCurveExistence(InDeltaChange, [&RegisteredCurves, &TemporaryCurves](const FCurveModelID& InCurveId)
+	{
+		return RegisteredCurves.Contains(InCurveId) || TemporaryCurves.Contains(InCurveId);
+	});
+	
+	// The keys of TMap<FCurveModelIDs> are completely distinct (we constructed them that way), so we can call ApplyChange / RevertChange twice here.
+	// So there is no risk of double applying / reverting any curve changes as the keys are distinct.
+	InProcessCurves(RegisteredCurves);
+	InProcessCurves(TemporaryCurves);
+}
 }
 	
 void FGenericCurveChangeCommand::Apply(UObject* Object)
 {
 	if (const TSharedPtr<FCurveEditor> CurveEditor = GetCurveEditor())
 	{
-		Private::ValidateCurveExistence(CurveEditor->GetCurves(), DeltaChange);
-		GenericCurveChange::ApplyChange(CurveEditor->GetCurves(), DeltaChange);
+		Private::ProcessExistingAndMissingCurves(*CurveEditor, DeltaChange, [this](const TMap<FCurveModelID, TUniquePtr<FCurveModel>>& InCurves)
+		{
+			GenericCurveChange::ApplyChange(InCurves, DeltaChange);
+		});
 	}
 }
 
@@ -59,8 +124,10 @@ void FGenericCurveChangeCommand::Revert(UObject* Object)
 {
 	if (const TSharedPtr<FCurveEditor> CurveEditor = GetCurveEditor())
 	{
-		Private::ValidateCurveExistence(CurveEditor->GetCurves(), DeltaChange);
-		GenericCurveChange::RevertChange(CurveEditor->GetCurves(), DeltaChange);
+		Private::ProcessExistingAndMissingCurves(*CurveEditor, DeltaChange, [this](const TMap<FCurveModelID, TUniquePtr<FCurveModel>>& InCurves)
+		{
+			GenericCurveChange::RevertChange(InCurves, DeltaChange);
+		});
 	}
 }
 
diff --git a/Engine/Source/Editor/CurveEditor/Private/Modification/Resolution/CurveModelLookUpInfo.cpp b/Engine/Source/Editor/CurveEditor/Private/Modification/Resolution/CurveModelLookUpInfo.cpp
new file mode 100644
--- /dev/null
+++ b/Engine/Source/Editor/CurveEditor/Private/Modification/Resolution/CurveModelLookUpInfo.cpp
@@ -0,0 +1,16 @@
+﻿// Copyright Epic Games, Inc. All Rights Reserved.
+
+#include "Modification/Resolution/CurveModelLookUpInfo.h"
+
+#include "CurveEditor.h"
+#include "CurveModelIdCache.h"
+
+namespace UE::CurveEditor
+{
+const FCurveMetaDataIdentifiers* FCurveModelLookUpArgs::GetCurveMetaData() const
+{
+	FCurveModelIdCache* Cache = CurveEditor.GetCurveMetaDataCache();
+	checkSlow(Cache); // Did not call FCurveEditor::InitCurveEditor?
+	return Cache ? FindMetaData(*Cache, CurveID) : nullptr;
+}
+}
diff --git a/Engine/Source/Editor/CurveEditor/Private/Modification/Resolution/ResolveCurveModelPasskey.h b/Engine/Source/Editor/CurveEditor/Private/Modification/Resolution/ResolveCurveModelPasskey.h
new file mode 100644
--- /dev/null
+++ b/Engine/Source/Editor/CurveEditor/Private/Modification/Resolution/ResolveCurveModelPasskey.h
@@ -0,0 +1,29 @@
+﻿// Copyright Epic Games, Inc. All Rights Reserved.
+
+#pragma once
+
+#include "CurveEditor.h"
+
+namespace UE::CurveEditor
+{
+/** Tries to re-create a FCurveModel that was in this FCurveEditor previously but has since been removed. */
+TUniquePtr<FCurveModel> ResolveCurveModel(const FCurveEditor& InCurveEditor, const FCurveModelID& InCurveId);
+	
+/** Implements passkey pattern to restrict access to FCurveEditor::ResolveCurveModel. */
+class FResolveCurveModelPasskey
+{
+	// Effectively make FCurveEditor::ResolveCurve available everywhere and only in CurveEditor.
+	// The whole point of this contraption is to avoid exposing FCurveEditor::ResolveCurve as part of the public API interface.
+	friend TUniquePtr<FCurveModel> ResolveCurveModel(const FCurveEditor& InCurveEditor, const FCurveModelID& InCurveId);
+	
+	static TUniquePtr<FCurveModel> ResolveCurveModel(const FCurveEditor& InCurveEditor, const FCurveModelID& InCurveId)
+	{
+		return InCurveEditor.ResolveCurve(InCurveId);
+	}
+};
+
+inline TUniquePtr<FCurveModel> ResolveCurveModel(const FCurveEditor& InCurveEditor, const FCurveModelID& InCurveId)
+{
+	return FResolveCurveModelPasskey::ResolveCurveModel(InCurveEditor, InCurveId);
+}
+}
diff --git a/Engine/Source/Editor/CurveEditor/Private/Modification/Resolution/ResolveCurveModelUtils.cpp b/Engine/Source/Editor/CurveEditor/Private/Modification/Resolution/ResolveCurveModelUtils.cpp
new file mode 100644
--- /dev/null
+++ b/Engine/Source/Editor/CurveEditor/Private/Modification/Resolution/ResolveCurveModelUtils.cpp
@@ -0,0 +1,8 @@
+﻿// Copyright Epic Games, Inc. All Rights Reserved.
+
+#include "ResolveCurveModelUtils.h"
+
+namespace UE::CurveEditor
+{
+	
+}
\ No newline at end of file
diff --git a/Engine/Source/Editor/CurveEditor/Private/Modification/Resolution/ResolveCurveModelUtils.h b/Engine/Source/Editor/CurveEditor/Private/Modification/Resolution/ResolveCurveModelUtils.h
new file mode 100644
--- /dev/null
+++ b/Engine/Source/Editor/CurveEditor/Private/Modification/Resolution/ResolveCurveModelUtils.h
@@ -0,0 +1,16 @@
+﻿// Copyright Epic Games, Inc. All Rights Reserved.
+
+#pragma once
+
+#include "Templates/UniquePtr.h"
+
+class FCurveModel;
+class FCurveEditor;
+struct FCurveModelID;
+
+namespace UE::CurveEditor
+{
+
+	
+	
+}
diff --git a/Engine/Source/Editor/CurveEditor/Public/CurveEditor.h b/Engine/Source/Editor/CurveEditor/Public/CurveEditor.h
--- a/Engine/Source/Editor/CurveEditor/Public/CurveEditor.h
+++ b/Engine/Source/Editor/CurveEditor/Public/CurveEditor.h
@@ -59,6 +59,7 @@ struct FCurveEditorSnapMetrics;
 struct FGeometry;
 namespace UE::CurveEditor { class FSelectionCleanser; }
 namespace UE::CurveEditor { struct FCurveModelIdCache; }
+namespace UE::CurveEditor { class FResolveCurveModelPasskey; }
 
 DECLARE_DELEGATE_OneParam(FOnSetBoolean, bool)
 DECLARE_MULTICAST_DELEGATE_OneParam(FOnActiveToolChanged, FCurveEditorToolID)
@@ -148,6 +149,8 @@ public:
 	UE_API virtual int32 GetSupportedTangentTypes();
 
 	UE::CurveEditor::FTransactionManager* GetTransactionManager() const { return TransactionManager.Get(); }
+	/** @return Metadata cache for reusing FCurveModelIDs when re-adding curves. For now, not API exported as it's for internal CurveEditor use only. */
+	UE::CurveEditor::FCurveModelIdCache* GetCurveMetaDataCache() const { return CurveModelIdCache.Get(); }
 
 public:
 
@@ -287,12 +290,13 @@ public:
 	 * @return a ptr to the curve if found, nullptr otherwise
 	 */
 	UE_API FCurveModel* FindCurve(FCurveModelID CurveID) const;
+	
 	/**
 	* Add a new curve to this editor
 	*/
 	UE_API FCurveModelID AddCurve(TUniquePtr<FCurveModel>&& InCurve);
 	UE_API FCurveModelID AddCurveForTreeItem(TUniquePtr<FCurveModel>&& InCurve, FCurveEditorTreeItemID TreeItemID);
-
+	
 	/**
 	* Remove a curve from this editor.
 	*/
@@ -773,6 +777,19 @@ protected:
 	 */
 	TPimplPtr<UE::CurveEditor::FCurveModelIdCache> CurveModelIdCache;
 
+private:
+
+	friend class UE::CurveEditor::FResolveCurveModelPasskey;
+	
+	/**
+	 * Internally used for editing FCurveModels after they've been removed FCurveEditor::CurveData (e.g. once they're no longer visible).
+	 * Important for applying undo / redo commands on curves that are not visible when the user issues the undo / redo operation.
+	 */
+	UE::CurveEditor::FResolveCurveModelDelegate ResolveCurveModelDelegate;
+	
+	/** Tries to re-create a FCurveModel that was in this FCurveEditor previously but has since been removed. */
+	TUniquePtr<FCurveModel> ResolveCurve(const FCurveModelID& InCurveID) const;
+
 private:
 
 	
diff --git a/Engine/Source/Editor/CurveEditor/Public/CurveEditorDelegates.h b/Engine/Source/Editor/CurveEditor/Public/CurveEditorDelegates.h
new file mode 100644
--- /dev/null
+++ b/Engine/Source/Editor/CurveEditor/Public/CurveEditorDelegates.h
@@ -0,0 +1,22 @@
+﻿// Copyright Epic Games, Inc. All Rights Reserved.
+
+#pragma  once
+
+#include "Delegates/Delegate.h"
+
+class FCurveModel;
+
+namespace UE::CurveEditor
+{
+struct FCurveModelLookUpArgs;
+	
+/**
+ * Tries to resolve a FCurveModel.
+ * 
+ * This is intended for asking the external systems owning the FCurveEditor instance to construct the temporary FCurveModel API,
+ * e.g. when perform an undo / redo operation for a FCurveModel that is not currently visible (FCurveEditor::CurveData only contains visible curves).
+ *
+ * Can return nullptr if the curve cannot be resolved.
+ */
+DECLARE_DELEGATE_RetVal_OneParam(TUniquePtr<FCurveModel>, FResolveCurveModelDelegate, const FCurveModelLookUpArgs& /*InArgs*/);
+}
diff --git a/Engine/Source/Editor/CurveEditor/Public/ICurveEditorModule.h b/Engine/Source/Editor/CurveEditor/Public/ICurveEditorModule.h
--- a/Engine/Source/Editor/CurveEditor/Public/ICurveEditorModule.h
+++ b/Engine/Source/Editor/CurveEditor/Public/ICurveEditorModule.h
@@ -2,6 +2,7 @@
 
 #pragma  once
 
+#include "CurveEditorDelegates.h"
 #include "Misc/Attribute.h"
 #include "Modules/ModuleManager.h"
 #include "Modules/ModuleInterface.h"
@@ -40,6 +41,13 @@ struct FCurveEditorInitParams
 
 	/** Optional. If set, defines multipliers for zooming. */
 	TAttribute<const FCurveEditorZoomScaleConfig*> ZoomScalingAttr;
+
+	/**
+	 * Internally used for editing FCurveModels after they've been removed FCurveEditor::CurveData (e.g. once they're no longer visible).
+	 * Important for applying undo / redo commands on curves that are not visible when the user issues the undo / redo operation.
+	 * You can skip providing this if you never call FCurveEditor::RemoveCurve.
+	 */
+	UE::CurveEditor::FResolveCurveModelDelegate ResolveCurveModelDelegate;
 };
 
 /**
diff --git a/Engine/Source/Editor/CurveEditor/Public/Modification/Resolution/CurveModelLookUpInfo.h b/Engine/Source/Editor/CurveEditor/Public/Modification/Resolution/CurveModelLookUpInfo.h
new file mode 100644
--- /dev/null
+++ b/Engine/Source/Editor/CurveEditor/Public/Modification/Resolution/CurveModelLookUpInfo.h
@@ -0,0 +1,37 @@
+﻿// Copyright Epic Games, Inc. All Rights Reserved.
+
+#pragma  once
+
+#include "CurveEditorTypes.h"
+
+class FCurveEditor;
+class FCurveModel;
+
+#define UE_API CURVEEDITOR_API
+
+namespace UE::CurveEditor
+{
+struct FCurveMetaDataIdentifiers;
+
+/**
+ * These arguments are used to temporarily resolve FCurveModel after it has been removed from FCurveEditor.
+ * 
+ * Remember, FCurveModel::CurveData only holds the curves that are currently visible.
+ * Some systems, e.g. undo / redo with FGenericCurveCommandChange, require a curve model later.
+ */
+struct FCurveModelLookUpArgs
+{
+	/** The curve editor instance we're operating on */
+	const FCurveEditor& CurveEditor;
+	/** The curve trying to be resolved. */
+	FCurveModelID CurveID;
+
+	explicit FCurveModelLookUpArgs(const FCurveEditor& CurveEditor, const FCurveModelID& CurveID) : CurveEditor(CurveEditor), CurveID(CurveID)
+	{}
+
+	/** @return Metadata information the curve had when it was last added FCurveEditor::AddCurve. */
+	CURVEEDITOR_API const FCurveMetaDataIdentifiers* GetCurveMetaData() const;
+};
+}
+
+#undef UE_API
\ No newline at end of file
diff --git a/Engine/Source/Editor/Sequencer/Private/MVVM/CurveEditorExtension.cpp b/Engine/Source/Editor/Sequencer/Private/MVVM/CurveEditorExtension.cpp
--- a/Engine/Source/Editor/Sequencer/Private/MVVM/CurveEditorExtension.cpp
+++ b/Engine/Source/Editor/Sequencer/Private/MVVM/CurveEditorExtension.cpp
@@ -20,6 +20,8 @@
 #include "Sequencer.h"
 #include "SequencerCommands.h"
 #include "Menus/SequencerToolbarUtils.h"
+#include "Modification/Resolution/CurveMetaDataIdentifiers.h"
+#include "Modification/Resolution/CurveModelLookUpInfo.h"
 #include "Modification/Utils/ScopedSelectionChange.h"
 #include "Toolkits/IToolkitHost.h"
 #include "Tree/SCurveEditorTree.h"
@@ -253,6 +255,45 @@ public:
 	}
 };
 
+/** Constructs a FCurveModel by looking up the channel from the owning UMovieSceneSection and ChannelName. */
+static TUniquePtr<FCurveModel> ResolveCurveEditorModel(const CurveEditor::FCurveModelLookUpArgs& InArgs, TWeakPtr<FSequencer> OwningSequencer)
+{
+	const TSharedPtr<ISequencer> SequencerPin = OwningSequencer.Pin();
+	if (!SequencerPin)
+	{
+		return nullptr;
+	}
+	
+	const CurveEditor::FCurveMetaDataIdentifiers* CurveMetadata = InArgs.GetCurveMetaData();
+	if (!CurveMetadata)
+	{
+		return nullptr;
+	}
+
+	UMovieSceneSection* OwningSection = Cast<UMovieSceneSection>(CurveMetadata->Owner.Get());
+	if (!OwningSection)
+	{
+		return nullptr;
+	}
+
+	const FMovieSceneChannelHandle ChannelHandle = OwningSection->GetChannelProxy().GetChannelByName(CurveMetadata->ChannelName);
+	ISequencerModule& SequencerModule = FModuleManager::LoadModuleChecked<ISequencerModule>("Sequencer");
+	ISequencerChannelInterface* EditorInterface = SequencerModule.FindChannelEditorInterface(ChannelHandle.GetChannelTypeName());
+	if (!EditorInterface)
+	{
+		return nullptr;
+	}
+	
+	const FMovieSceneChannelMetaData* ChannelMetadata = ChannelHandle.GetMetaData();
+	FCreateCurveEditorModelParams Params = {
+		OwningSection,
+		ChannelMetadata->WeakOwningObject.Get(),
+		SequencerPin.ToSharedRef()
+	};
+
+	return EditorInterface->CreateCurveEditorModel_Raw(ChannelHandle, Params);
+}
+
 const FName FCurveEditorExtension::CurveEditorTabName = FName(TEXT("SequencerGraphEditor"));
 
 FCurveEditorExtension::FCurveEditorExtension()
@@ -297,6 +338,7 @@ void FCurveEditorExtension::CreateCurveEditor(const FTimeSliderArgs& TimeSliderA
 		{
 			return &SequencerSettings->GetCurveEditorZoomScaling();
 		});
+		CurveEditorInitParams.ResolveCurveModelDelegate.BindStatic(&ResolveCurveEditorModel, Sequencer.ToWeakPtr());
 		
 		CurveEditorModel = MakeShared<FSequencerCurveEditor>(Sequencer, TimeSliderArgs.NumericTypeInterface);
 		CurveEditorModel->SetBounds(MakeUnique<FSequencerCurveEditorBounds>(Sequencer.ToSharedRef()));
