Anim In Engine: Filters can work on all open curves in addition to selected keys

#Synced-CL 51508179

diff --git a/Engine/Plugins/MovieScene/SequencerScripting/Source/SequencerScriptingEditor/Private/SequencerCurveEditorObject.cpp b/Engine/Plugins/MovieScene/SequencerScripting/Source/SequencerScriptingEditor/Private/SequencerCurveEditorObject.cpp
--- a/Engine/Plugins/MovieScene/SequencerScripting/Source/SequencerScriptingEditor/Private/SequencerCurveEditorObject.cpp
+++ b/Engine/Plugins/MovieScene/SequencerScripting/Source/SequencerScriptingEditor/Private/SequencerCurveEditorObject.cpp
@@ -340,29 +340,51 @@ void USequencerCurveEditorObject::ApplyFilter(UCurveEditorFilterBase* Filter)
 {
 	using namespace UE::CurveEditor;
 
-	TSharedPtr<FCurveEditor> CurveEditor =  GetCurveEditor();
+	TSharedPtr<FCurveEditor> CurveEditor = GetCurveEditor();
 
-	if (Filter && CurveEditor.IsValid())
+	if (!Filter || !CurveEditor.IsValid())
 	{
-		const FText TransactionText = FText::Format(
-			NSLOCTEXT("CurveEditorFilterApply", "Filter Curves", "Filtered {0}|plural(one=Curve, other=Curves)"),
-			CurveEditor->GetSelection().Count()
-		);
-		const FScopedTransaction Transaction(TransactionText);
-
-		Filter->InitializeFilter(CurveEditor.ToSharedRef());
-		const TMap<FCurveModelID, FKeyHandleSet>& SelectedKeys = CurveEditor->GetSelection().GetAll();
-		TMap<FCurveModelID, FKeyHandleSet> OutKeysToSelect;
-		{
-			const FScopedCurveChange KeyChange(FCurvesSnapshotBuilder(CurveEditor, SelectedKeys));
+		return;
+	}
+
+	const TMap<FCurveModelID, FKeyHandleSet>& SelectedKeys = CurveEditor->GetSelection().GetAll();
 
-			Filter->ApplyFilter(CurveEditor.ToSharedRef(), SelectedKeys, OutKeysToSelect);
+	// Build working key set: selection if non-empty, else all visible curves
+	TMap<FCurveModelID, FKeyHandleSet> AllCurveKeys;
+	const TMap<FCurveModelID, FKeyHandleSet>* KeysToOperateOn = &SelectedKeys;
+	const bool bSelectionEmpty = SelectedKeys.IsEmpty();
+	if (bSelectionEmpty)
+	{
+		for (const TTuple<FCurveModelID, TUniquePtr<FCurveModel>>& CurvePair : CurveEditor->GetCurves())
+		{
+			FKeyHandleSet& HandleSet = AllCurveKeys.Add(CurvePair.Key);
+			for (const FKeyHandle& Handle : CurvePair.Value->GetAllKeys())
+			{
+				HandleSet.Add(Handle, ECurvePointType::Key);
+			}
 		}
+		KeysToOperateOn = &AllCurveKeys;
+	}
+
+	const int32 SelectionCount = CurveEditor->GetSelection().Count();
+	const FText TransactionText = SelectionCount > 0
+		? FText::Format(NSLOCTEXT("CurveEditorFilterApply", "Filter Curves", "Filtered {0}|plural(one=Curve, other=Curves)"), SelectionCount)
+		: NSLOCTEXT("CurveEditorFilterApply", "Filter All Visible Curves", "Filtered All Visible Curves");
+
+	const FScopedTransaction Transaction(TransactionText);
+
+	Filter->InitializeFilter(CurveEditor.ToSharedRef());
+	TMap<FCurveModelID, FKeyHandleSet> OutKeysToSelect;
+	{
+		const FScopedCurveChange KeyChange(FCurvesSnapshotBuilder(CurveEditor, *KeysToOperateOn));
+		Filter->ApplyFilter(CurveEditor.ToSharedRef(), *KeysToOperateOn, OutKeysToSelect);
+	}
+
+	if (bSelectionEmpty == false) //only reselect if we had selection
+	{
 		const FScopedSelectionChange SelectionTransaction(CurveEditor);
 
-		// Clear their selection and then set it to the keys the filter thinks you should have selected.
 		CurveEditor->GetSelection().Clear();
-
 		for (const TTuple<FCurveModelID, FKeyHandleSet>& OutSet : OutKeysToSelect)
 		{
 			CurveEditor->GetSelection().Add(OutSet.Key, ECurvePointType::Key, OutSet.Value.AsArray());
diff --git a/Engine/Source/Editor/CurveEditor/Private/Filters/CurveEditorBakeFilter.cpp b/Engine/Source/Editor/CurveEditor/Private/Filters/CurveEditorBakeFilter.cpp
--- a/Engine/Source/Editor/CurveEditor/Private/Filters/CurveEditorBakeFilter.cpp
+++ b/Engine/Source/Editor/CurveEditor/Private/Filters/CurveEditorBakeFilter.cpp
@@ -54,7 +54,7 @@ void UCurveEditorBakeFilter::InitializeFilter_Impl(TSharedRef<FCurveEditor> InCu
 
 bool UCurveEditorBakeFilter::CanApplyFilter_Impl(TSharedRef<FCurveEditor> InCurveEditor)
 {
-	return bCustomRangeOverride || InCurveEditor->GetSelection().Count() > 0;
+	return bCustomRangeOverride || InCurveEditor->GetSelection().Count() > 0 || InCurveEditor->GetCurves().Num() > 0;
 }
 
 void UCurveEditorBakeFilter::ApplyFilter_Impl(TSharedRef<FCurveEditor> InCurveEditor, const TMap<FCurveModelID, FKeyHandleSet>& InKeysToOperateOn, TMap<FCurveModelID, FKeyHandleSet>& OutKeysToSelect)
diff --git a/Engine/Source/Editor/CurveEditor/Private/Filters/CurveEditorEulerFilter.cpp b/Engine/Source/Editor/CurveEditor/Private/Filters/CurveEditorEulerFilter.cpp
--- a/Engine/Source/Editor/CurveEditor/Private/Filters/CurveEditorEulerFilter.cpp
+++ b/Engine/Source/Editor/CurveEditor/Private/Filters/CurveEditorEulerFilter.cpp
@@ -94,12 +94,19 @@ void UCurveEditorEulerFilter::ApplyFilter_Impl(TSharedRef<FCurveEditor> InCurveE
 bool UCurveEditorEulerFilter::CanApplyFilter_Impl(TSharedRef<FCurveEditor> InCurveEditor)
 {
 	const TMap<FCurveModelID, FKeyHandleSet>& Selection = InCurveEditor->GetSelection().GetAll();
-	const bool bHasRotationCurves = !Selection.IsEmpty()
-		&& Algo::AnyOf(InCurveEditor->GetSelection().GetAll(),
-		[&InCurveEditor](const TPair<FCurveModelID, FKeyHandleSet>& Pair)
+	if (!Selection.IsEmpty())
+	{
+		return Algo::AnyOf(Selection,
+			[&InCurveEditor](const TPair<FCurveModelID, FKeyHandleSet>& Pair)
+			{
+				const FCurveModel* Curve = InCurveEditor->FindCurve(Pair.Key);
+				return Curve && Curve->GetIntentionName().Contains("Rotation");
+			});
+	}
+	// No selection — allow if any visible curve is a rotation curve
+	return Algo::AnyOf(InCurveEditor->GetCurves(),
+		[](const TTuple<FCurveModelID, TUniquePtr<FCurveModel>>& CurvePair)
 		{
-			const FCurveModel* Curve = InCurveEditor->FindCurve(Pair.Key);
-			return Curve && Curve->GetIntentionName().Contains("Rotation");
+			return CurvePair.Value && CurvePair.Value->GetIntentionName().Contains("Rotation");
 		});
-	return bHasRotationCurves;
 }
diff --git a/Engine/Source/Editor/CurveEditor/Private/Filters/CurveEditorFilterBase.cpp b/Engine/Source/Editor/CurveEditor/Private/Filters/CurveEditorFilterBase.cpp
--- a/Engine/Source/Editor/CurveEditor/Private/Filters/CurveEditorFilterBase.cpp
+++ b/Engine/Source/Editor/CurveEditor/Private/Filters/CurveEditorFilterBase.cpp
@@ -48,7 +48,7 @@ FSlateIcon UCurveEditorFilterBase::GetIcon(const TSubclassOf<UCurveEditorFilterB
 
 bool UCurveEditorFilterBase::CanApplyFilter_Impl(TSharedRef<FCurveEditor> InCurveEditor)
 {
-	return InCurveEditor->GetSelection().Count() > 0;
+	return InCurveEditor->GetSelection().Count() > 0 || InCurveEditor->GetCurves().Num() > 0;
 }
 
 #undef LOCTEXT_NAMESPACE
diff --git a/Engine/Source/Editor/CurveEditor/Private/Filters/CurveEditorSmartSnapFilter.cpp b/Engine/Source/Editor/CurveEditor/Private/Filters/CurveEditorSmartSnapFilter.cpp
--- a/Engine/Source/Editor/CurveEditor/Private/Filters/CurveEditorSmartSnapFilter.cpp
+++ b/Engine/Source/Editor/CurveEditor/Private/Filters/CurveEditorSmartSnapFilter.cpp
@@ -20,5 +20,5 @@ void UCurveEditorSmartSnapFilter::ApplyFilter_Impl(
 
 bool UCurveEditorSmartSnapFilter::CanApplyFilter_Impl(TSharedRef<FCurveEditor> InCurveEditor)
 {
-	return UE::CurveEditor::CanSmartSnapSelection(InCurveEditor->GetSelection());
+	return UE::CurveEditor::CanSmartSnapSelection(InCurveEditor->GetSelection()) || InCurveEditor->GetCurves().Num() > 0;
 }
diff --git a/Engine/Source/Editor/CurveEditor/Private/Filters/FilterUtils.cpp b/Engine/Source/Editor/CurveEditor/Private/Filters/FilterUtils.cpp
--- a/Engine/Source/Editor/CurveEditor/Private/Filters/FilterUtils.cpp
+++ b/Engine/Source/Editor/CurveEditor/Private/Filters/FilterUtils.cpp
@@ -6,6 +6,8 @@
 #include "Modification/Utils/ScopedSelectionChange.h"
 
 #include "Containers/Map.h"
+#include "CurveEditor.h"
+#include "CurveModel.h"
 #include "Modification/Utils/ScopedCurveChange.h"
 #include "Templates/SharedPointer.h"
 
@@ -13,30 +15,71 @@ namespace UE::CurveEditor::FilterUtils
 {
 	void ApplyFilter(const TSharedRef<FCurveEditor>& InCurveEditor, UCurveEditorFilterBase& InFilter)
 	{
-		ApplyFilter(InCurveEditor, InFilter, InCurveEditor->Selection.GetAll());
+		const TMap<FCurveModelID, FKeyHandleSet>& SelectedKeys = InCurveEditor->Selection.GetAll();
+		auto HasSelectedKeys = [&SelectedKeys]()
+			{
+				if (SelectedKeys.IsEmpty())
+				{
+					return false;
+				}
+				else
+				{
+					for (const TTuple<FCurveModelID, FKeyHandleSet>& Pair : SelectedKeys)
+					{
+						if (Pair.Value.Num() > 0)
+						{
+							return true;
+						}
+					}
+				}
+				return false;
+
+			};
+		const bool bHasSelectedKeys = HasSelectedKeys();
+		if (bHasSelectedKeys == false)
+		{
+			// No keys selected — operate on all visible curves
+			TMap<FCurveModelID, FKeyHandleSet> AllCurveKeys;
+			for (const TTuple<FCurveModelID, TUniquePtr<FCurveModel>>& CurvePair : InCurveEditor->GetCurves())
+			{
+				FKeyHandleSet& HandleSet = AllCurveKeys.Add(CurvePair.Key);
+				for (const FKeyHandle& Handle : CurvePair.Value->GetAllKeys())
+				{
+					HandleSet.Add(Handle, ECurvePointType::Key);
+				}
+			}
+			ApplyFilter(InCurveEditor, InFilter, AllCurveKeys, bHasSelectedKeys);
+		}
+		else
+		{
+			ApplyFilter(InCurveEditor, InFilter, SelectedKeys, bHasSelectedKeys);
+		}
 	}
 
 	void ApplyFilter(
 		const TSharedRef<FCurveEditor>& InCurveEditor,
 		UCurveEditorFilterBase& InFilter,
-		const TMap<FCurveModelID, FKeyHandleSet>& InSelectedKeys
+		const TMap<FCurveModelID, FKeyHandleSet>& InKeysToFilter,
+		bool bResetSelection
 		)
 	{
 		using namespace UE::CurveEditor;
 		
 		TMap<FCurveModelID, FKeyHandleSet> OutKeysToSelect;
 		{
-			const FScopedCurveChange KeyChange(FCurvesSnapshotBuilder(InCurveEditor, InSelectedKeys));
-			InFilter.ApplyFilter(InCurveEditor, InSelectedKeys, OutKeysToSelect);
+			const FScopedCurveChange KeyChange(FCurvesSnapshotBuilder(InCurveEditor, InKeysToFilter));
+			InFilter.ApplyFilter(InCurveEditor, InKeysToFilter, OutKeysToSelect);
 		}
-		
-		const FScopedSelectionChange Transaction(InCurveEditor);
-		// Clear their selection and then set it to the keys the filter thinks you should have selected.
-		InCurveEditor->GetSelection().Clear();
-
-		for (const TTuple<FCurveModelID, FKeyHandleSet>& OutSet : OutKeysToSelect)
+		if (bResetSelection)
 		{
-			InCurveEditor->GetSelection().Add(OutSet.Key, ECurvePointType::Key, OutSet.Value.AsArray());
+			const FScopedSelectionChange Transaction(InCurveEditor);
+			// Clear their selection and then set it to the keys the filter thinks you should have selected.
+			InCurveEditor->GetSelection().Clear();
+
+			for (const TTuple<FCurveModelID, FKeyHandleSet>& OutSet : OutKeysToSelect)
+			{
+				InCurveEditor->GetSelection().Add(OutSet.Key, ECurvePointType::Key, OutSet.Value.AsArray());
+			}
 		}
 	}
 }
\ No newline at end of file
diff --git a/Engine/Source/Editor/CurveEditor/Private/Filters/FilterUtils.h b/Engine/Source/Editor/CurveEditor/Private/Filters/FilterUtils.h
--- a/Engine/Source/Editor/CurveEditor/Private/Filters/FilterUtils.h
+++ b/Engine/Source/Editor/CurveEditor/Private/Filters/FilterUtils.h
@@ -20,6 +20,7 @@ namespace UE::CurveEditor::FilterUtils
 	void ApplyFilter(
 		const TSharedRef<FCurveEditor>& InCurveEditor,
 		UCurveEditorFilterBase& InFilter,
-		const TMap<FCurveModelID, FKeyHandleSet>& InSelectedKeys
+		const TMap<FCurveModelID, FKeyHandleSet>& InSelectedKeys,
+		bool bResetSelection
 		);
 }
diff --git a/Engine/Source/Editor/CurveEditor/Private/Filters/SCurveEditorFilterPanel.cpp b/Engine/Source/Editor/CurveEditor/Private/Filters/SCurveEditorFilterPanel.cpp
--- a/Engine/Source/Editor/CurveEditor/Private/Filters/SCurveEditorFilterPanel.cpp
+++ b/Engine/Source/Editor/CurveEditor/Private/Filters/SCurveEditorFilterPanel.cpp
@@ -163,12 +163,20 @@ FReply SCurveEditorFilterPanel::OnApplyClicked()
 	TSharedPtr<FCurveEditor> CurveEditor = WeakCurveEditor.Pin();
 	if (CurveEditor && DetailsView->GetSelectedObjects().Num() > 0)
 	{
-		const FText TransactionText = FText::Format(
-			NSLOCTEXT("CurveEditorFilterApply", "Filter Curves", "Filtered {0}|plural(one=Curve, other=Curves)"),
-			CurveEditor->GetSelection().Count()
-			);
+		const int32 SelectionCount = CurveEditor->GetSelection().Count();
+		FText TransactionText;
+		if (SelectionCount > 0)
+		{
+			TransactionText = FText::Format(
+				NSLOCTEXT("CurveEditorFilterApply", "Filter Curves", "Filtered {0}|plural(one=Curve, other=Curves)"),
+				SelectionCount);
+		}
+		else
+		{
+			TransactionText = NSLOCTEXT("CurveEditorFilterApply", "Filter All Visible Curves", "Filtered All Visible Curves");
+		}
 		const FScopedTransaction Transaction(TransactionText);
-		
+
 		UCurveEditorFilterBase* Filter = Cast<UCurveEditorFilterBase>(DetailsView->GetSelectedObjects()[0]);
 		UE::CurveEditor::FilterUtils::ApplyFilter(CurveEditor.ToSharedRef(), *Filter);
 	}
@@ -256,12 +264,17 @@ FText SCurveEditorFilterPanel::GetCurrentFilterText() const
 	TSharedPtr<FCurveEditor> CurveEditor = WeakCurveEditor.Pin();
 	if (CurveEditor.IsValid())
 	{
-		int32 SelectedKeysCount = CurveEditor->GetSelection().Count();
+		const int32 SelectedKeysCount = CurveEditor->GetSelection().Count();
 		if (SelectedKeysCount > 0)
 		{
 			const FText FormatString = NSLOCTEXT("SCurveEditorFilterPanel", "CurrentFilterWithKeysFormat", "Current Filter: {0} ({1} Keys)");
 			return FText::Format(FormatString, CurrentFilterName, SelectedKeysCount);
 		}
+		else if (CurveEditor->GetCurves().Num() > 0)
+		{
+			const FText FormatString = NSLOCTEXT("SCurveEditorFilterPanel", "CurrentFilterAllCurvesFormat", "Current Filter: {0} (All Visible Curves)");
+			return FText::Format(FormatString, CurrentFilterName);
+		}
 		else
 		{
 			const FText FormatString = NSLOCTEXT("SCurveEditorFilterPanel", "CurrentFilterNoKeysFormat", "Current Filter: {0} (No Keys Selected)");
