Anim In Engine: Add options to the smart reduce filter to remove keys from static channels and use the default instead(which is set when the key is set) and only apply the filter to these static channels (e.g. you want ot only nuke scale channels).

#Synced-CL 51508179

diff --git a/Engine/Source/Editor/CurveEditor/Private/Filters/CurveEditorSmartReduceFilter.cpp b/Engine/Source/Editor/CurveEditor/Private/Filters/CurveEditorSmartReduceFilter.cpp
--- a/Engine/Source/Editor/CurveEditor/Private/Filters/CurveEditorSmartReduceFilter.cpp
+++ b/Engine/Source/Editor/CurveEditor/Private/Filters/CurveEditorSmartReduceFilter.cpp
@@ -119,6 +119,17 @@ void UCurveEditorSmartReduceFilter::SmartReduce(FCurveModel* Curve, const FSmart
 	const bool bNeedToTestExisting, FKeyHandleSet& OutHandleSet)
 {
 
+	// bOnlyReduceStaticChannels: skip curves whose value range is not flat
+	if (InParams.bOnlyReduceStaticChannels)
+	{
+		double MinValue, MaxValue;
+		Curve->GetValueRange(MinValue, MaxValue);
+		if (!FMath::IsNearlyEqual(MinValue, MaxValue, UE_KINDA_SMALL_NUMBER))
+		{
+			return;
+		}
+	}
+
 	TArray<FKeyHandle> KeyHandles;
 	TArray<FKeyPosition> SelectedKeyPositions;
 	double MinKey = TNumericLimits<double>::Max(), MaxKey = TNumericLimits<double>::Lowest();
@@ -301,6 +312,36 @@ void UCurveEditorSmartReduceFilter::SmartReduce(FCurveModel* Curve, const FSmart
 	{
 		OutHandleSet.Add(NewHandle, ECurvePointType::Key);
 	}
+
+
+	// bDeleteStaticChannels: if exactly 2 keys remain with the same value, remove them
+	// and let RemoveKeys set the channel default via InCurrentTime
+	if (InParams.bDeleteStaticChannels)
+	{
+		TArray<FKeyHandle> AllRemainingKeys;
+		Curve->GetAllKeys(AllRemainingKeys);
+
+		if (AllRemainingKeys.Num() == 2)
+		{
+			TArray<FKeyPosition> KeyPositions;
+			KeyPositions.SetNum(2);
+			Curve->GetKeyPositions(AllRemainingKeys, KeyPositions);
+
+			if (FMath::IsNearlyEqual(
+				KeyPositions[0].OutputValue,
+				KeyPositions[1].OutputValue,
+				UE_KINDA_SMALL_NUMBER))
+			{
+				// Remove both keys; InCurrentTime = first key time so the default
+				// is evaluated to the constant value and stored on the channel
+				Curve->RemoveKeys(AllRemainingKeys, KeyPositions[0].InputValue);
+				for (const FKeyHandle& RemainingHandle : AllRemainingKeys)
+				{
+					OutHandleSet.Remove(RemainingHandle, ECurvePointType::Key);
+				}
+			}
+		}
+	}
 }
 
 void UCurveEditorSmartReduceFilter::ApplyFilter_Impl(TSharedRef<FCurveEditor> InCurveEditor, const TMap<FCurveModelID, FKeyHandleSet>& InKeysToOperateOn, TMap<FCurveModelID, FKeyHandleSet>& OutKeysToSelect)
@@ -313,10 +354,12 @@ void UCurveEditorSmartReduceFilter::ApplyFilter_Impl(TSharedRef<FCurveEditor> In
 		{
 			continue;
 		}
+
 		const bool bNeedToTestExisting = true;
 		FKeyHandleSet& OutHandleSet = OutKeysToSelect.Add(Pair.Key);
 
 		TOptional<FKeyHandleSet> KeyHandleSet = Pair.Value;
 		SmartReduce(Curve, SmartReduceParams, KeyHandleSet, bNeedToTestExisting, OutHandleSet);
+
 	}
 }
diff --git a/Engine/Source/Editor/CurveEditor/Public/Filters/CurveEditorSmartReduceFilter.h b/Engine/Source/Editor/CurveEditor/Public/Filters/CurveEditorSmartReduceFilter.h
--- a/Engine/Source/Editor/CurveEditor/Public/Filters/CurveEditorSmartReduceFilter.h
+++ b/Engine/Source/Editor/CurveEditor/Public/Filters/CurveEditorSmartReduceFilter.h
@@ -29,10 +29,20 @@ struct FSmartReduceParams
 	UPROPERTY(EditAnywhere, BlueprintReadWrite, Category = "Settings")
 	FFrameRate SampleRate = FFrameRate(30, 1);
 
+	/** Remove all keys from static channels, and set that value as the default */
+	UPROPERTY(EditAnywhere, BlueprintReadWrite, Category = "Settings")
+	bool bDeleteStaticChannels = false;
+
+	/** Only reduce static channels */
+	UPROPERTY(EditAnywhere, BlueprintReadWrite, Category = "Settings")
+	bool bOnlyReduceStaticChannels = false;
+
 	void Reset()
 	{
 		TolerancePercentage = 5.0;
 		SampleRate = FFrameRate(30, 1);
+		bDeleteStaticChannels = false;
+		bOnlyReduceStaticChannels = false;
 	}
 };
 
