Curve Editor: The normalized now draws the initial view correctly

#Synced-CL 52063343

diff --git a/Engine/Source/Editor/CurveEditor/Private/Views/SCurveEditorViewNormalized.cpp b/Engine/Source/Editor/CurveEditor/Private/Views/SCurveEditorViewNormalized.cpp
--- a/Engine/Source/Editor/CurveEditor/Private/Views/SCurveEditorViewNormalized.cpp
+++ b/Engine/Source/Editor/CurveEditor/Private/Views/SCurveEditorViewNormalized.cpp
@@ -40,10 +40,9 @@ constexpr float NormalizedPadding = 10.f;
 
 void SCurveEditorViewNormalized::Construct(const FArguments& InArgs, TWeakPtr<FCurveEditor> InCurveEditor)
 {
-	//when created we set the output bounds to be fixed since otherwise it may get fitted
-	bFixedOutputBounds = true;
 	bAllowModelViewTransforms = false;
-	FrameVertical(-0.1, 1.1);
+	OutputMin = -0.1;
+	OutputMax =  1.1;
 
 	SInteractiveCurveEditorView::Construct(InArgs, InCurveEditor);
 
@@ -220,8 +219,8 @@ void SCurveEditorViewNormalized::InternalUpdateViewToTransformCurves()
 	const TSharedPtr<FCurveEditor> CurveEditor = WeakCurveEditor.Pin();
 	if (CurveEditor.IsValid())
 	{
-		bFixedOutputBounds = false;
-		double InputMin = 0, InputMax = 1;
+		double InputMin = 0.0;
+		double InputMax = 1.0;
 		GetInputBounds(InputMin, InputMax);
 		UpdateViewToTransformCurves(InputMin, InputMax);
 	}
