Anim In Engine: Hotkeys for 1) Zoom To/Frame Selection Range Command in Sequencer and Curve Editor, 2) Select Keys in Selection Range in Curve Editor 3) Adding Keys Only onto selected curves in the Curve Editor, and 4) Reseting Selected Key Values to default channel value.  Also Curve Editor setting "Cubic Add Prefers User Tangent" to use User Tangents when adding keys to cubic curves.

#Synced-CL 52171595

diff --git a/Engine/Plugins/Editor/CurveEditorTools/Source/CurveEditorTools/Private/CurveEditorToolCommands.cpp b/Engine/Plugins/Editor/CurveEditorTools/Source/CurveEditorTools/Private/CurveEditorToolCommands.cpp
--- a/Engine/Plugins/Editor/CurveEditorTools/Source/CurveEditorTools/Private/CurveEditorToolCommands.cpp
+++ b/Engine/Plugins/Editor/CurveEditorTools/Source/CurveEditorTools/Private/CurveEditorToolCommands.cpp
@@ -9,6 +9,10 @@ void FCurveEditorToolCommands::RegisterCommands()
 	// Focus Tools
 	UI_COMMAND(SetFocusPlaybackTime, "Focus Playback Time", "Focuses the Curve Editor on the current Playback Time without changing zoom level.", EUserInterfaceActionType::Button, FInputChord());
 	UI_COMMAND(SetFocusPlaybackRange, "Focus Playback Range", "Focuses the Curve Editor on the current Playback Range with zoom based on visible curves.", EUserInterfaceActionType::Button, FInputChord(EKeys::A));
+	UI_COMMAND(SetFocusSelectionRange, "Focus Selection Range", "Focuses the Curve Editor on the current Selection Range with zoom based on visible curves.", EUserInterfaceActionType::Button, FInputChord());
+	UI_COMMAND(SetKeysInSelectionRange, "Select Keys in Selection Range",
+		"Selects all keys within the current Selection Range (or Playback Range if no selection range is set).",
+		EUserInterfaceActionType::Button, FInputChord());
 
 	// Tool Modes
 	UI_COMMAND(ActivateTransformTool, "Transform", "The Transform tool allows translation, scale and rotation of selected keys.", EUserInterfaceActionType::RadioButton, FInputChord(EModifierKey::Control, EKeys::T));
diff --git a/Engine/Plugins/Editor/CurveEditorTools/Source/CurveEditorTools/Private/CurveEditorToolCommands.h b/Engine/Plugins/Editor/CurveEditorTools/Source/CurveEditorTools/Private/CurveEditorToolCommands.h
--- a/Engine/Plugins/Editor/CurveEditorTools/Source/CurveEditorTools/Private/CurveEditorToolCommands.h
+++ b/Engine/Plugins/Editor/CurveEditorTools/Source/CurveEditorTools/Private/CurveEditorToolCommands.h
@@ -24,6 +24,8 @@ public:
 
 	TSharedPtr<FUICommandInfo> SetFocusPlaybackTime;
 	TSharedPtr<FUICommandInfo> SetFocusPlaybackRange;
+	TSharedPtr<FUICommandInfo> SetFocusSelectionRange;
+	TSharedPtr<FUICommandInfo> SetKeysInSelectionRange;
 	TSharedPtr<FUICommandInfo> ActivateTransformTool;
 	TSharedPtr<FUICommandInfo> ActivateRetimeTool;
 	TSharedPtr<FUICommandInfo> ActivateMultiScaleTool;
diff --git a/Engine/Plugins/Editor/CurveEditorTools/Source/CurveEditorTools/Private/Extensions/CurveEditorFocusExtension.cpp b/Engine/Plugins/Editor/CurveEditorTools/Source/CurveEditorTools/Private/Extensions/CurveEditorFocusExtension.cpp
--- a/Engine/Plugins/Editor/CurveEditorTools/Source/CurveEditorTools/Private/Extensions/CurveEditorFocusExtension.cpp
+++ b/Engine/Plugins/Editor/CurveEditorTools/Source/CurveEditorTools/Private/Extensions/CurveEditorFocusExtension.cpp
@@ -6,6 +6,8 @@
 #include "CurveEditor.h"
 #include "ICurveEditorBounds.h"
 #include "AnimatedRange.h"
+#include "CurveModel.h"
+#include "CurveEditorSelection.h"
 
 void FCurveEditorFocusExtension::BindCommands(TSharedRef<FUICommandList> CommandBindings)
 {
@@ -13,9 +15,17 @@ void FCurveEditorFocusExtension::BindCommands(TSharedRef<FUICommandList> Command
 	FCanExecuteAction CanFocus = FCanExecuteAction::CreateSP(this, &FCurveEditorFocusExtension::CanUseFocusExtensions);
 	FExecuteAction FocusPlaybackTime = FExecuteAction::CreateSP(this, &FCurveEditorFocusExtension::FocusPlaybackTime);
 	FExecuteAction FocusPlaybackRange = FExecuteAction::CreateSP(this, &FCurveEditorFocusExtension::FocusPlaybackRange);
+	FExecuteAction FocusSelectionRange = FExecuteAction::CreateSP(this, &FCurveEditorFocusExtension::FocusSelectionRange);
 
 	CommandBindings->MapAction(FCurveEditorToolCommands::Get().SetFocusPlaybackTime, FocusPlaybackTime, CanFocus);
 	CommandBindings->MapAction(FCurveEditorToolCommands::Get().SetFocusPlaybackRange, FocusPlaybackRange, CanFocus);
+	CommandBindings->MapAction(FCurveEditorToolCommands::Get().SetFocusSelectionRange, FocusSelectionRange, CanFocus);
+
+	FExecuteAction SelectKeysInSelectionRangeAction = FExecuteAction::CreateSP(
+		this, &FCurveEditorFocusExtension::SelectKeysInSelectionRange);
+	CommandBindings->MapAction(
+		FCurveEditorToolCommands::Get().SetKeysInSelectionRange,
+		SelectKeysInSelectionRangeAction, CanFocus);
 }
 
 void FCurveEditorFocusExtension::FocusPlaybackRange()
@@ -51,6 +61,36 @@ void FCurveEditorFocusExtension::FocusPlaybackRange()
 	}
 }
 
+void FCurveEditorFocusExtension::FocusSelectionRange()
+{
+	TSharedPtr<FCurveEditor> CurveEditor = WeakCurveEditor.Pin();
+	if (CurveEditor)
+	{
+		TSharedPtr<ITimeSliderController> TimeSliderController = CurveEditor->GetTimeSliderController();
+		if (TimeSliderController)
+		{
+			TRange<FFrameNumber> Range = TimeSliderController->GetSelectionRange();
+			if (Range.IsEmpty())
+			{
+				Range = TimeSliderController->GetPlayRange();
+			}
+
+			CurveEditor->ZoomToFit(EAxisList::Type::Y);
+
+			FFrameNumber InFrame  = Range.GetLowerBoundValue();
+			FFrameNumber OutFrame = Range.GetUpperBoundValue();
+			double InTime  = InFrame  / TimeSliderController->GetTickResolution();
+			double OutTime = OutFrame / TimeSliderController->GetTickResolution();
+
+			double Padding = (OutTime - InTime) * 0.05;
+			InTime  -= Padding;
+			OutTime += Padding;
+
+			CurveEditor->GetBounds().SetInputBounds(InTime, OutTime);
+		}
+	}
+}
+
 void FCurveEditorFocusExtension::FocusPlaybackTime()
 {
 	// This takes your current view range (both x and y) and shifts it to be centered on the playback head.
@@ -97,6 +137,65 @@ void FCurveEditorFocusExtension::FocusPlaybackTime()
 	}
 }
 
+void FCurveEditorFocusExtension::SelectKeysInSelectionRange()
+{
+	TSharedPtr<FCurveEditor> CurveEditor = WeakCurveEditor.Pin();
+	if (!CurveEditor)
+	{
+		return;
+	}
+
+	TSharedPtr<ITimeSliderController> TimeSliderController =
+		CurveEditor->GetTimeSliderController();
+	if (!TimeSliderController)
+	{
+		return;
+	}
+
+	// Resolve the time range: selection range first, play range as fallback
+	TRange<FFrameNumber> FrameRange = TimeSliderController->GetSelectionRange();
+	if (FrameRange.IsEmpty())
+	{
+		FrameRange = TimeSliderController->GetPlayRange();
+	}
+	if (FrameRange.IsEmpty())
+	{
+		return;
+	}
+
+	FFrameRate TickResolution = TimeSliderController->GetTickResolution();
+	double InTime  = FrameRange.GetLowerBoundValue() / TickResolution;
+	double OutTime = FrameRange.GetUpperBoundValue() / TickResolution;
+
+	// Collect keys that fall within the range for every edited curve
+	TMap<FCurveModelID, TArray<FKeyHandle>> KeysPerCurve;
+	for (FCurveModelID ID : CurveEditor->GetEditedCurves())
+	{
+		if (FCurveModel* Curve = CurveEditor->FindCurve(ID))
+		{
+			TArray<FKeyHandle> KeyHandles;
+			Curve->GetKeys(InTime, OutTime,
+				TNumericLimits<double>::Lowest(), TNumericLimits<double>::Max(),
+				KeyHandles);
+			if (KeyHandles.Num() > 0)
+			{
+				KeysPerCurve.Add(ID, MoveTemp(KeyHandles));
+			}
+		}
+	}
+
+	// Only modify selection if at least one key was found
+	if (KeysPerCurve.Num() > 0)
+	{
+		FCurveEditorSelection& Selection = CurveEditor->GetSelection();
+		Selection.Clear();
+		for (auto& Pair : KeysPerCurve)
+		{
+			Selection.Add(Pair.Key, ECurvePointType::Key, Pair.Value);
+		}
+	}
+}
+
 bool FCurveEditorFocusExtension::CanUseFocusExtensions() const
 {
 	TSharedPtr<FCurveEditor> CurveEditor = WeakCurveEditor.Pin();
diff --git a/Engine/Plugins/Editor/CurveEditorTools/Source/CurveEditorTools/Private/Extensions/CurveEditorFocusExtension.h b/Engine/Plugins/Editor/CurveEditorTools/Source/CurveEditorTools/Private/Extensions/CurveEditorFocusExtension.h
--- a/Engine/Plugins/Editor/CurveEditorTools/Source/CurveEditorTools/Private/Extensions/CurveEditorFocusExtension.h
+++ b/Engine/Plugins/Editor/CurveEditorTools/Source/CurveEditorTools/Private/Extensions/CurveEditorFocusExtension.h
@@ -21,6 +21,10 @@ public:
 private:
 	/** Focuses the curve editor around the current Playback Range of the linked Time Controller. */
 	void FocusPlaybackRange();
+	/** Focuses the curve editor around the current Selection Range of the linked Time Controller (falls back to Playback Range if empty). */
+	void FocusSelectionRange();
+	/** Selects all keys within the current Selection Range (or Playback Range if no selection range is set). */
+	void SelectKeysInSelectionRange();
 	/** Focuses the curve editor around the current Playback Time of the linked Time Controller without changing zoom level. */
 	void FocusPlaybackTime();
 
diff --git a/Engine/Source/Editor/CurveEditor/Private/CurveEditorCommands.cpp b/Engine/Source/Editor/CurveEditor/Private/CurveEditorCommands.cpp
--- a/Engine/Source/Editor/CurveEditor/Private/CurveEditorCommands.cpp
+++ b/Engine/Source/Editor/CurveEditor/Private/CurveEditorCommands.cpp
@@ -80,6 +80,8 @@ void FCurveEditorCommands::RegisterCommands()
 	UI_COMMAND(PasteRelative, "Paste Relative", "Pastes the keys from the clipboard and aligns them to the nearest key to the left of the scrubber.", EUserInterfaceActionType::Button, FInputChord(EModifierKey::Alt, EKeys::V));
 	
 	UI_COMMAND(AddKeyToAllCurves, "Add Key", "Add a new key to all curves at the current time.", EUserInterfaceActionType::Button, FInputChord(EKeys::Enter) );
+	UI_COMMAND(AddKeyToSelectedCurves, "Add Key to Selected Curves", "Add a new key to all curves that have selected keys at the current time.", EUserInterfaceActionType::Button, FInputChord());
+	UI_COMMAND(ResetKeysToDefaultValues, "Reset to Default", "Reset selected keys to their default value", EUserInterfaceActionType::Button, FInputChord());
 
 	// Curve Editor Colors
 	UI_COMMAND(SetRandomCurveColorsForSelected, "Set Random Curve Colors", "Set random colors on the selected curves. Note they are stored in the Level Sequence Actor Editor Preferences.", EUserInterfaceActionType::Button, FInputChord());
diff --git a/Engine/Source/Editor/CurveEditor/Private/CurveEditorSettings.cpp b/Engine/Source/Editor/CurveEditor/Private/CurveEditorSettings.cpp
--- a/Engine/Source/Editor/CurveEditor/Private/CurveEditorSettings.cpp
+++ b/Engine/Source/Editor/CurveEditor/Private/CurveEditorSettings.cpp
@@ -21,6 +21,7 @@ UCurveEditorSettings::UCurveEditorSettings()
 	SnapAxis = ECurveEditorSnapAxis::CESA_None;
 	bSnapTimeToSelection = false;
 	bScrubTimeStartFromCursor = true;
+	bCubicAddKeyPrefersUserTangent = false;
 
 	SelectionColor = FLinearColor::White;
 	ParentSpaceCustomColor = FLinearColor(.93, .31, .19); //pastel orange
@@ -162,6 +163,20 @@ void UCurveEditorSettings::SetShowBufferedCurves(bool InbShowBufferedCurves)
 	}
 }
 
+bool UCurveEditorSettings::GetCubicAddKeyPrefersUserTangent() const
+{
+	return bCubicAddKeyPrefersUserTangent;
+}
+
+void UCurveEditorSettings::SetCubicAddKeyPrefersUserTangent(bool bInCubicAddKeyPrefersUserTangent)
+{
+	if (bCubicAddKeyPrefersUserTangent != bInCubicAddKeyPrefersUserTangent)
+	{
+		bCubicAddKeyPrefersUserTangent = bInCubicAddKeyPrefersUserTangent;
+		SaveConfig();
+	}
+}
+
 bool UCurveEditorSettings::GetShowCurveEditorCurveToolTips() const
 {
 	return bShowCurveEditorCurveToolTips;
diff --git a/Engine/Source/Editor/CurveEditor/Private/Views/SInteractiveCurveEditorView.cpp b/Engine/Source/Editor/CurveEditor/Private/Views/SInteractiveCurveEditorView.cpp
--- a/Engine/Source/Editor/CurveEditor/Private/Views/SInteractiveCurveEditorView.cpp
+++ b/Engine/Source/Editor/CurveEditor/Private/Views/SInteractiveCurveEditorView.cpp
@@ -1991,7 +1991,8 @@ void SInteractiveCurveEditorView::RebindContextualActions(FVector2D InMousePosit
 	CommandList->UnmapAction(FCurveEditorCommands::Get().AddKeyHovered);
 	CommandList->UnmapAction(FCurveEditorCommands::Get().PasteKeysHovered);
 	CommandList->UnmapAction(FCurveEditorCommands::Get().AddKeyToAllCurves);
-
+	CommandList->UnmapAction(FCurveEditorCommands::Get().AddKeyToSelectedCurves);
+	CommandList->UnmapAction(FCurveEditorCommands::Get().ResetKeysToDefaultValues);
 	CommandList->UnmapAction(FCurveEditorCommands::Get().BufferVisibleCurves);
 	CommandList->UnmapAction(FCurveEditorCommands::Get().SwapBufferedCurves);
 	CommandList->UnmapAction(FCurveEditorCommands::Get().ApplyBufferedCurves);
@@ -2007,7 +2008,16 @@ void SInteractiveCurveEditorView::RebindContextualActions(FVector2D InMousePosit
 		CommandList->MapAction(FCurveEditorCommands::Get().PasteKeysHovered, FExecuteAction::CreateSP(this, &SInteractiveCurveEditorView::PasteKeys, HoveredCurveSet));
 	}
 
-	CommandList->MapAction(FCurveEditorCommands::Get().AddKeyToAllCurves, FExecuteAction::CreateSP(this, &SInteractiveCurveEditorView::AddKeyAtScrubTime, TSet<FCurveModelID>()));
+	CommandList->MapAction(FCurveEditorCommands::Get().AddKeyToAllCurves, FExecuteAction::CreateSP(this, &SInteractiveCurveEditorView::AddKeyAtScrubTime, false));
+
+	CommandList->MapAction(
+		FCurveEditorCommands::Get().AddKeyToSelectedCurves,
+		FExecuteAction::CreateSP(this, &SInteractiveCurveEditorView::AddKeyAtScrubTime, true /*use selection*/));
+
+	CommandList->MapAction(
+		FCurveEditorCommands::Get().ResetKeysToDefaultValues,
+		FExecuteAction::CreateSP(this, &SInteractiveCurveEditorView::ResetKeysToDefaultValues, true /*use selection*/));
+		
 
 	// Buffer Curves. Can only act on buffered curves if curves are selected in the tree or the curve has selected keys.
 	CommandList->MapAction(FCurveEditorCommands::Get().BufferVisibleCurves, FExecuteAction::CreateSP(this, &SInteractiveCurveEditorView::BufferCurves), FCanExecuteAction::CreateSP(this, &SInteractiveCurveEditorView::CanBufferedCurves));
@@ -2066,18 +2076,88 @@ FKeyAttributes SInteractiveCurveEditorView::GetDefaultKeyAttributesForCurveTime(
 	FKeyAttributes KeyAttributes = CurveEditor.GetDefaultKeyAttribute().Get();
 
 	// Give the model a chance to override the default interpolation modes
-    const TPair<ERichCurveInterpMode, ERichCurveTangentMode> Modes = CurveModel.GetInterpolationMode(
+    TPair<ERichCurveInterpMode, ERichCurveTangentMode> Modes = CurveModel.GetInterpolationMode(
     	EvalTime,
     	KeyAttributes.HasInterpMode() ? KeyAttributes.GetInterpMode() : RCIM_Linear,
     	KeyAttributes.HasTangentMode() ? KeyAttributes.GetTangentMode() : RCTM_Auto
     	);
+
+	if (const UCurveEditorSettings* Settings = GetDefault<UCurveEditorSettings>())
+	{
+		if (Settings->GetCubicAddKeyPrefersUserTangent())
+		{
+			if (Modes.Key == RCIM_Cubic)
+			{
+				Modes.Value = RCTM_User;
+			}
+		}
+	}
     KeyAttributes.SetInterpMode(Modes.Key);
     KeyAttributes.SetTangentMode(Modes.Value);
 	
     return KeyAttributes;
 }
 
-void SInteractiveCurveEditorView::AddKeyAtScrubTime(TSet<FCurveModelID> ForCurves)
+void SInteractiveCurveEditorView::ResetKeysToDefaultValues(bool bUseSelection)
+{
+	using namespace UE::CurveEditor;
+
+	TSharedPtr<FCurveEditor> CurveEditor = WeakCurveEditor.Pin();
+	if (!CurveEditor.IsValid())
+	{
+		return;
+	}
+
+	TSet<FCurveModelID> CurvesToReset;
+	if (bUseSelection == false)
+	{
+		CurvesToReset = CurveEditor->GetEditedCurves();
+	}
+	else
+	{
+		for (const auto& Pair : CurveEditor->GetSelection().GetAll())
+		{
+			CurvesToReset.Add(Pair.Key);
+		}
+	}
+
+	FScopedTransaction Transaction(LOCTEXT("ResetKeysToDefaultValues", "Set Default Value")); 
+	FScopedSelectionChange SelectionChange(WeakCurveEditor);
+	FScopedCurveChange KeyChange(
+		FCurvesSnapshotBuilder(WeakCurveEditor, CurvesToReset, ECurveChangeFlags::MoveKeys)
+	);
+
+	FCurveEditorSelection& Selection = CurveEditor->GetSelection();
+	TArray<FKeyHandle> KeyHandles;
+	TArray<FKeyPosition> SelectedKeyPositions;
+	for (FCurveModelID CurveID : CurvesToReset)
+	{
+		if (FCurveModel* CurveModel = CurveEditor->FindCurve(CurveID))
+		{
+			if (CurveModel->HasDefaultValue())
+			{
+				double DefaultValue = CurveModel->GetDefaultValue();
+				if (const FKeyHandleSet* SelectionSet = Selection.FindForCurve(CurveID))
+				{
+					if (SelectionSet->Num() > 0)
+					{
+						KeyHandles = SelectionSet->AsArray();
+						SelectedKeyPositions.SetNum(KeyHandles.Num());
+						CurveModel->GetKeyPositions(KeyHandles, SelectedKeyPositions);
+						for (FKeyPosition &Key : SelectedKeyPositions)
+						{
+							Key.OutputValue = DefaultValue;
+						}
+						CurveModel->SetKeyPositions(KeyHandles, SelectedKeyPositions);
+					}
+				}
+			}
+		}
+	}
+}
+
+
+void SInteractiveCurveEditorView::AddKeyAtScrubTime(bool bUseSelection)
 {
 	TSharedPtr<FCurveEditor> CurveEditor = WeakCurveEditor.Pin();
 	if (!CurveEditor.IsValid())
@@ -2086,13 +2166,16 @@ void SInteractiveCurveEditorView::AddKeyAtScrubTime(TSet<FCurveModelID> ForCurve
 	}
 
 	TSet<FCurveModelID> CurvesToAddTo;
-	if (ForCurves.Num() == 0)
+	if (bUseSelection == false)
 	{
 		CurvesToAddTo = CurveEditor->GetEditedCurves();
 	}
 	else
 	{
-		CurvesToAddTo = ForCurves;
+		for (const auto& Pair : CurveEditor->GetSelection().GetAll())
+		{
+			CurvesToAddTo.Add(Pair.Key);
+		}
 	}
 
 	// If they don't have a time slider controller then we fall back to using mouse position.
diff --git a/Engine/Source/Editor/CurveEditor/Public/CurveEditor.h b/Engine/Source/Editor/CurveEditor/Public/CurveEditor.h
--- a/Engine/Source/Editor/CurveEditor/Public/CurveEditor.h
+++ b/Engine/Source/Editor/CurveEditor/Public/CurveEditor.h
@@ -594,6 +594,16 @@ public:
 	 */
 	UE_API void DeleteSelection();
 
+	/**
+	 * Reset selected keys to their curve's default value, for curves that support it.
+	 */
+	UE_API void ResetSelectedKeysToDefaultValues();
+
+	/**
+	 * Returns true if any selected key belongs to a curve that has a default value.
+	 */
+	UE_API bool CanResetSelectedKeysToDefaultValues() const;
+
 	/**
 	 * Flatten the tangents on the selected keys
 	 */
diff --git a/Engine/Source/Editor/CurveEditor/Public/CurveEditorCommands.h b/Engine/Source/Editor/CurveEditor/Public/CurveEditorCommands.h
--- a/Engine/Source/Editor/CurveEditor/Public/CurveEditorCommands.h
+++ b/Engine/Source/Editor/CurveEditor/Public/CurveEditorCommands.h
@@ -98,6 +98,9 @@ public:
 	TSharedPtr<FUICommandInfo> PasteRelative;
 
 	TSharedPtr<FUICommandInfo> AddKeyToAllCurves;
+	TSharedPtr<FUICommandInfo> AddKeyToSelectedCurves;
+	TSharedPtr<FUICommandInfo> ResetKeysToDefaultValues;
+
 
 	TSharedPtr<FUICommandInfo> SetViewModeAbsolute;
 	TSharedPtr<FUICommandInfo> SetViewModeStacked;
diff --git a/Engine/Source/Editor/CurveEditor/Public/CurveEditorSettings.h b/Engine/Source/Editor/CurveEditor/Public/CurveEditorSettings.h
--- a/Engine/Source/Editor/CurveEditor/Public/CurveEditorSettings.h
+++ b/Engine/Source/Editor/CurveEditor/Public/CurveEditorSettings.h
@@ -149,6 +149,12 @@ public:
 	/** Sets whether or not to show buffered curves in the curve editor. */
 	UE_API void SetShowBufferedCurves(bool InbShowBufferedCurves);
 
+
+	/** Get whether inserted key will be user tangent when possible. */
+	UE_API bool GetCubicAddKeyPrefersUserTangent() const;
+	/** Sets wGet whether inserted key will be user tangent when possible. */
+	UE_API void SetCubicAddKeyPrefersUserTangent(bool bInCubicAddKeyPrefersUserTangent);
+
 	/** Gets whether or not to show curve tool tips in the curve editor. */
 	UE_API bool GetShowCurveEditorCurveToolTips() const;
 	/** Sets whether or not to show curve tool tips in the curve editor. */
@@ -287,6 +293,9 @@ protected:
 	UPROPERTY(config, EditAnywhere, Category = "Curve Editor")
 	bool bScrubTimeStartFromCursor;
 
+	/** When enabled, inserted key will be user tangent when possible */
+	UPROPERTY(config, EditAnywhere, Category = "Curve Editor")
+	bool bCubicAddKeyPrefersUserTangent;
 	/**
 	 * Determines how close you must move the selection marque to the center of a point in order to select it.
 	 * 
diff --git a/Engine/Source/Editor/CurveEditor/Public/CurveModel.h b/Engine/Source/Editor/CurveEditor/Public/CurveModel.h
--- a/Engine/Source/Editor/CurveEditor/Public/CurveModel.h
+++ b/Engine/Source/Editor/CurveEditor/Public/CurveModel.h
@@ -398,13 +398,24 @@ public:
 		return FTransform2d();
 	}
 
-	/** 
+	/**
 	 * Returns whether the curve model should be edited or not
 	 */
 	virtual bool IsReadOnly() const
 	{
 		return false;
 	}
+
+	/**
+	 * Returns whether this curve has a default value that selected keys can be reset to.
+	 */
+	virtual bool HasDefaultValue() const { return false; }
+
+	/**
+	 * Returns the default value for this curve. Only meaningful when HasDefaultValue() returns true.
+	 */
+	virtual double GetDefaultValue() const { return 0.0; }
+
 	/**
 	* Get the UObject that owns this CurveModel, for example for Sequencer this would be the UMovieSceneSection
 	*/
diff --git a/Engine/Source/Editor/CurveEditor/Public/Views/SInteractiveCurveEditorView.h b/Engine/Source/Editor/CurveEditor/Public/Views/SInteractiveCurveEditorView.h
--- a/Engine/Source/Editor/CurveEditor/Public/Views/SInteractiveCurveEditorView.h
+++ b/Engine/Source/Editor/CurveEditor/Public/Views/SInteractiveCurveEditorView.h
@@ -183,11 +183,13 @@ private:
 	UE_API double GetTangentValue(const double InTime, const double InValue, FCurveModel* CurveToAddTo, double DeltaTime) const;
 
 	/*~ Command binding callbacks */
-	UE_API void AddKeyAtScrubTime(TSet<FCurveModelID> ForCurves);
+	UE_API void AddKeyAtScrubTime(bool bUseSelection);
+	UE_DEPRECATED(5.8, "Please, use AddKeyAtScrubTime that takes bUseSelection parameter")
+	void AddKeyAtScrubTime(TSet<FCurveModelID> ForCurves) { AddKeyAtScrubTime(ForCurves.Num() > 0); }
 	UE_API void AddKeyAtMousePosition(TSet<FCurveModelID> ForCurves);
 	UE_API void AddKeyAtTime(const TSet<FCurveModelID>& ToCurves, double InTime);
 	UE_API void PasteKeys(TSet<FCurveModelID> ToCurves);
-
+	UE_API void ResetKeysToDefaultValues(bool bUseSelection);
 	UE_API void OnCurveEditorToolChanged(FCurveEditorToolID InToolId);
 	UE_API void OnShowValueIndicatorsChanged();
 
diff --git a/Engine/Source/Editor/MovieSceneTools/Private/Channels/BezierChannelCurveModel.cpp b/Engine/Source/Editor/MovieSceneTools/Private/Channels/BezierChannelCurveModel.cpp
--- a/Engine/Source/Editor/MovieSceneTools/Private/Channels/BezierChannelCurveModel.cpp
+++ b/Engine/Source/Editor/MovieSceneTools/Private/Channels/BezierChannelCurveModel.cpp
@@ -689,6 +689,29 @@ void FBezierChannelCurveModel<ChannelType, ChannelValue, KeyType>::SetKeyValue(i
 	}
 }
 
+template<typename ChannelType, typename ChannelValue, typename KeyType>
+bool FBezierChannelCurveModel<ChannelType, ChannelValue, KeyType>::HasDefaultValue() const
+{
+	if (const ChannelType* Channel = this->GetChannelHandle().Get())
+	{
+		return Channel->GetDefault().IsSet();
+	}
+	return false;
+}
+
+template<typename ChannelType, typename ChannelValue, typename KeyType>
+double FBezierChannelCurveModel<ChannelType, ChannelValue, KeyType>::GetDefaultValue() const
+{
+	if (const ChannelType* Channel = this->GetChannelHandle().Get())
+	{
+		if (HasDefaultValue())
+		{
+			return static_cast<double>(Channel->GetDefault().GetValue());
+		}
+	}
+	return 0.0;
+}
+
 template class FBezierChannelCurveModel<FMovieSceneFloatChannel, FMovieSceneFloatValue, float>;
 template class FBezierChannelCurveModel<FMovieSceneDoubleChannel, FMovieSceneDoubleValue, double>;
 
diff --git a/Engine/Source/Editor/MovieSceneTools/Public/Channels/BezierChannelCurveModel.h b/Engine/Source/Editor/MovieSceneTools/Public/Channels/BezierChannelCurveModel.h
--- a/Engine/Source/Editor/MovieSceneTools/Public/Channels/BezierChannelCurveModel.h
+++ b/Engine/Source/Editor/MovieSceneTools/Public/Channels/BezierChannelCurveModel.h
@@ -79,6 +79,9 @@ public:
 	virtual void GetValueRange(double& MinValue, double& MaxValue) const override;
 	virtual void GetValueRange(double InMinTime, double InMaxTime, double& MinValue, double& MaxValue) const override;
 
+	virtual bool HasDefaultValue() const override;
+	virtual double GetDefaultValue() const override;
+
 protected:
 	// FChannelCurveModel
 	virtual double GetKeyValue(TArrayView<const ChannelValue> Values, int32 Index) const override;
diff --git a/Engine/Source/Editor/Sequencer/Private/Sequencer.cpp b/Engine/Source/Editor/Sequencer/Private/Sequencer.cpp
--- a/Engine/Source/Editor/Sequencer/Private/Sequencer.cpp
+++ b/Engine/Source/Editor/Sequencer/Private/Sequencer.cpp
@@ -7481,11 +7481,58 @@ void FSequencer::ZoomToFit()
 	}
 }
 
+void FSequencer::ZoomToFitSelectionRange()
+{
+	FFrameRate TickResolution = GetFocusedTickResolution();
+
+	// Use the movie scene selection range if set, otherwise fall back to time bounds
+	TRange<FFrameNumber> BoundsHull = TRange<FFrameNumber>::Empty();
+	if (const UMovieSceneSequence* FocusedSequence = GetFocusedMovieSceneSequence())
+	{
+		if (const UMovieScene* MovieScene = FocusedSequence->GetMovieScene())
+		{
+			BoundsHull = MovieScene->GetSelectionRange();
+		}
+	}
+	if (BoundsHull.IsEmpty())
+	{
+		BoundsHull = GetTimeBounds();
+	}
+
+	if (!BoundsHull.IsEmpty() && !BoundsHull.IsDegenerate())
+	{
+		const double Tolerance = KINDA_SMALL_NUMBER;
+
+		if (!ViewRangeBeforeZoom.IsEmpty() &&
+			FMath::IsNearlyEqual(BoundsHull.GetLowerBoundValue() / TickResolution,
+			                      GetViewRange().GetLowerBoundValue(), Tolerance) &&
+			FMath::IsNearlyEqual(BoundsHull.GetUpperBoundValue() / TickResolution,
+			                      GetViewRange().GetUpperBoundValue(), Tolerance))
+		{
+			SetViewRange(ViewRangeBeforeZoom, EViewRangeInterpolation::Animated);
+		}
+		else
+		{
+			ViewRangeBeforeZoom = GetViewRange();
+
+			TRange<double> BoundsHullSeconds = BoundsHull / TickResolution;
+			const double OutputViewSize = BoundsHullSeconds.Size<double>();
+			const double OutputChange   = OutputViewSize * 0.1f;
+
+			if (OutputChange > 0)
+			{
+				BoundsHullSeconds = UE::MovieScene::ExpandRange(BoundsHullSeconds, OutputChange);
+				SetViewRange(BoundsHullSeconds, EViewRangeInterpolation::Animated);
+			}
+		}
+	}
+}
+
 bool FSequencer::CanKeyProperty(FCanKeyPropertyParams CanKeyPropertyParams) const
 {
 	FPropertyPath PropertyPath;
 	return ObjectChangeListener->CanKeyProperty(CanKeyPropertyParams, PropertyPath);
-} 
+}
 
 
 void FSequencer::KeyProperty(FKeyPropertyParams KeyPropertyParams) 
@@ -12323,6 +12370,10 @@ void FSequencer::BindCommands()
 		Commands.ZoomToFit,
 		FExecuteAction::CreateSP( this, &FSequencer::ZoomToFit ) );
 
+	SequencerCommandBindings->MapAction(
+		Commands.ZoomToFitSelectionRange,
+		FExecuteAction::CreateSP( this, &FSequencer::ZoomToFitSelectionRange) );
+
 	SequencerCommandBindings->MapAction(
 		Commands.ZoomInViewRange,
 		FExecuteAction::CreateSP( this, &FSequencer::ZoomInViewRange ),
diff --git a/Engine/Source/Editor/Sequencer/Private/Sequencer.h b/Engine/Source/Editor/Sequencer/Private/Sequencer.h
--- a/Engine/Source/Editor/Sequencer/Private/Sequencer.h
+++ b/Engine/Source/Editor/Sequencer/Private/Sequencer.h
@@ -438,6 +438,9 @@ public:
 	/** Zooms to the edges of all currently selected sections and keys. */
 	void ZoomToFit();
 
+	/** Zooms to the movie scene selection range, or falls back to time bounds if empty. */
+	void ZoomToFitSelectionRange();
+
 	/** Gets the overlay fading animation curve lerp. */
 	float GetOverlayFadeCurve() const;
 
diff --git a/Engine/Source/Editor/Sequencer/Private/SequencerCommands.cpp b/Engine/Source/Editor/Sequencer/Private/SequencerCommands.cpp
--- a/Engine/Source/Editor/Sequencer/Private/SequencerCommands.cpp
+++ b/Engine/Source/Editor/Sequencer/Private/SequencerCommands.cpp
@@ -40,6 +40,7 @@ void FSequencerCommands::RegisterCommands()
 	UI_COMMAND( FocusPlaybackTime, "Focus Playback Time", "Focus the view range on the current playback time without changing zoom level", EUserInterfaceActionType::Button, FInputChord());
 	UI_COMMAND( ResetViewRange, "Reset View Range", "Reset view range to the playback range", EUserInterfaceActionType::Button, FInputChord(EKeys::Home) );
 	UI_COMMAND( ZoomToFit, "Zoom to Fit", "Zoom to Fit", EUserInterfaceActionType::Button, FInputChord(EKeys::F) );
+	UI_COMMAND( ZoomToFitSelectionRange, "Focus Selection Range", "Zoom the view to fit the selection range, or the full time bounds if no selection range is set.", EUserInterfaceActionType::Button, FInputChord() );
 	UI_COMMAND( ZoomInViewRange, "Zoom into the View Range", "Zoom into the view range", EUserInterfaceActionType::Button, FInputChord(EKeys::Equals) );
 	UI_COMMAND( ZoomOutViewRange, "Zoom out of the View Range", "Zoom out of the view range", EUserInterfaceActionType::Button, FInputChord(EKeys::Hyphen) );
 	UI_COMMAND( NavigateBackward, "Navigate Backward", "Go backward to the previously viewed shot/subsequence", EUserInterfaceActionType::Button, FInputChord(EKeys::ThumbMouseButton), FInputChord(EModifierKey::Alt, EKeys::Left));
diff --git a/Engine/Source/Editor/Sequencer/Public/SequencerCommands.h b/Engine/Source/Editor/Sequencer/Public/SequencerCommands.h
--- a/Engine/Source/Editor/Sequencer/Public/SequencerCommands.h
+++ b/Engine/Source/Editor/Sequencer/Public/SequencerCommands.h
@@ -108,6 +108,9 @@ public:
 	/** Zoom to fit the selected sections and keys */
 	TSharedPtr< FUICommandInfo > ZoomToFit;
 
+	/** Zoom to fit the selection range (falls back to time bounds if selection is empty) */
+	TSharedPtr< FUICommandInfo > ZoomToFitSelectionRange;
+
 	/** Zoom into the view range */
 	TSharedPtr< FUICommandInfo > ZoomInViewRange;
 
