Curve Editor: Fix FCurveEditorTree::AddItem returning stale pointers.

~FScopedCurveEditorTreeEventGuard executes FCurveEditorTree::Compact, which made NewItem become dangling.

#Synced-CL 52283655

diff --git a/Engine/Source/Editor/CurveEditor/Private/Tree/CurveEditorTree.cpp b/Engine/Source/Editor/CurveEditor/Private/Tree/CurveEditorTree.cpp
--- a/Engine/Source/Editor/CurveEditor/Private/Tree/CurveEditorTree.cpp
+++ b/Engine/Source/Editor/CurveEditor/Private/Tree/CurveEditorTree.cpp
@@ -171,22 +171,26 @@ const TMap<FCurveEditorTreeItemID, FCurveEditorTreeItem>& FCurveEditorTree::GetA
 
 FCurveEditorTreeItem* FCurveEditorTree::AddItem(FCurveEditorTreeItemID ParentID)
 {
-	FScopedCurveEditorTreeEventGuard EventGuard(this);
-
-	FCurveEditorTreeItemID NewItemID = NextTreeItemID;
+	const FCurveEditorTreeItemID NewItemID = NextTreeItemID;
+	{
+		const FScopedCurveEditorTreeEventGuard EventGuard(this);
 
-	FCurveEditorTreeItem& NewItem = Items.Add(NewItemID);
-	NewItem.ThisID   = NewItemID;
-	NewItem.ParentID = ParentID;
+		FCurveEditorTreeItem& NewItem = Items.Add(NewItemID);
+		NewItem.ThisID   = NewItemID;
+		NewItem.ParentID = ParentID;
 
-	FSortedCurveEditorTreeItems& ParentContainer = ParentID.IsValid() ? Items.FindChecked(ParentID).Children : RootItems;
-	ParentContainer.ChildIDs.Add(NewItemID);
-	ParentContainer.bRequiresSort = true;
+		FSortedCurveEditorTreeItems& ParentContainer = ParentID.IsValid() ? Items.FindChecked(ParentID).Children : RootItems;
+		ParentContainer.ChildIDs.Add(NewItemID);
+		ParentContainer.bRequiresSort = true;
 
-	++NextTreeItemID.Value;
-	++Events.OnItemsChanged.SerialNumber;
+		++NextTreeItemID.Value;
+		++Events.OnItemsChanged.SerialNumber;
+	}
 
-	return &NewItem;
+	// ~FScopedCurveEditorTreeEventGuard executes FCurveEditorTree::Compact, so NewItem may have been reallocated.
+	FCurveEditorTreeItem* Result = Items.Find(NewItemID);
+	ensure(Result);
+	return Result;
 }
 
 void FCurveEditorTree::RemoveItem(FCurveEditorTreeItemID ItemID, FCurveEditor* CurveEditor)
