Curve Editor: Fix tweening not flattening user tangents when tweening exactly 1 key.

This introduces a new algorithm for tweening a single key based on its neighbours. See details in code.

#Synced-CL 52544961

diff --git a/Engine/Plugins/Animation/TweeningUtils/Source/TweeningUtilsEditor/Public/Math/Models/TangentFlatteningTweenProxy.h b/Engine/Plugins/Animation/TweeningUtils/Source/TweeningUtilsEditor/Public/Math/Models/TangentFlatteningTweenProxy.h
--- a/Engine/Plugins/Animation/TweeningUtils/Source/TweeningUtilsEditor/Public/Math/Models/TangentFlatteningTweenProxy.h
+++ b/Engine/Plugins/Animation/TweeningUtils/Source/TweeningUtilsEditor/Public/Math/Models/TangentFlatteningTweenProxy.h
@@ -8,6 +8,8 @@
 #include "Misc/Mirror/TangentSelectionFlattener.h"
 #include <type_traits>
 
+#include "Misc/Mirror/TangentSingleKeyFlattener.h"
+
 namespace UE::TweeningUtilsEditor
 {
 /**
@@ -35,14 +37,26 @@ public:
 		if (const TSharedPtr<FCurveEditor> CurveEditorPin = WeakCurveEditor.Get().Pin())
 		{
 			TangentTweener.ResetFromSelection(*CurveEditorPin);
+			SingleKeyTangentTweener.ResetFromSelection(*CurveEditorPin);
 		}
 	}
 	virtual void BlendValues(float InNormalizedValue) override
 	{
+		using namespace CurveEditor;
+		
 		TBase::BlendValues(InNormalizedValue);
 		if (const TSharedPtr<FCurveEditor> CurveEditorPin = WeakCurveEditor.Get().Pin())
 		{
-			TangentTweener.ComputeMirroringParallel(*CurveEditorPin);
+			const double FinalBlendValue = TBase::ScaleBlendValue(InNormalizedValue);
+			
+			// Conceptually, we've got two horizontal, parallel edges: the top one, which moves InNormalizedValue goes from 0 to +/- 1, and the bottom edge at 0.
+			// The tangents move depending on the edge relative distance. During overshoot mode, if the value exceeds +/- 1, the top edge has crossed the 0 edge.
+			const bool bIsOvershooting = FMath::Abs(FinalBlendValue) > 1.f;
+			const bool bTopEdgeHasCrossedBottomEdge = bIsOvershooting;
+			TangentTweener.ComputeMirroringParallel(*CurveEditorPin, bTopEdgeHasCrossedBottomEdge);
+			
+			const ESide NeighborToBlendTowards = FinalBlendValue < 0.f ? ESide::Left : ESide::Right;
+			SingleKeyTangentTweener.ComputeMirroringParallel(*CurveEditorPin, NeighborToBlendTowards);
 		}
 	}
 	//~ Begin FTweenModel Interface
@@ -51,8 +65,10 @@ private:
 
 	/** Needed as arg for TangentTweener. */
 	const TAttribute<TWeakPtr<FCurveEditor>> WeakCurveEditor;
-	/** Implements the logic for flattening the tangents. */
+	/** Implements the logic for flattening tangents of consecutive keys. */
 	CurveEditor::FTangentSelectionFlattener TangentTweener;
+	/** Implements logic for flattening tangents of single keys */
+	CurveEditor::FTangentSingleKeyFlattener SingleKeyTangentTweener;
 };
 }
 
diff --git a/Engine/Source/Editor/CurveEditor/Private/Misc/Mirror/TangentSingleKeyFlattener.cpp b/Engine/Source/Editor/CurveEditor/Private/Misc/Mirror/TangentSingleKeyFlattener.cpp
new file mode 100644
--- /dev/null
+++ b/Engine/Source/Editor/CurveEditor/Private/Misc/Mirror/TangentSingleKeyFlattener.cpp
@@ -0,0 +1,239 @@
+﻿// Copyright Epic Games, Inc. All Rights Reserved.
+
+#include "Misc/Mirror/TangentSingleKeyFlattener.h"
+
+#include "Misc/Mirror/MirrorUtils.h"
+#include "Misc/Mirror/SingleKeyMirrorData.h"
+
+namespace UE::CurveEditor
+{
+namespace TangentFlattenDetail
+{
+static FKeyNeighborData AnalyzeKey(const FKeyHandle& InKey, const FCurveModel& InCurveModel)
+{
+	TOptional<FKeyHandle> Before, After;
+	InCurveModel.GetNeighboringKeys(InKey, Before, After);
+	
+	TArray<FKeyHandle, TInlineAllocator<3>> Keys;
+	TArray<FKeyPosition, TInlineAllocator<3>> Positions;
+	if (Before)
+	{
+		Keys.Add(*Before);
+	}
+	if (After)
+	{
+		Keys.Add(*After);
+	}
+	Keys.Add(InKey);
+	Positions.SetNumUninitialized(Keys.Num());
+	InCurveModel.GetKeyPositions(Keys, Positions);
+	
+	const int32 BeforeIndex = Before ? 0 : INDEX_NONE;
+	const int32 AfterIndex = BeforeIndex + 1;
+	const int32 CurrentIndex = Keys.Num() - 1;
+	return FKeyNeighborData
+	{
+		Positions[CurrentIndex].OutputValue,
+		Before ? Positions[BeforeIndex].OutputValue : TOptional<double>{},
+		After ? Positions[AfterIndex].OutputValue : TOptional<double>{},
+	};
+}
+
+/** Gets the neighbor to compute the initial interpolation distance to. */
+static TOptional<ESide> GetInitialAnchorSide(const FKeyNeighborData& InData)
+{
+	return InData.LeftKeyHeight 
+		? ESide::Left 
+		: InData.RightKeyHeight ? ESide::Right : TOptional<ESide>{};
+}
+
+static ESide Opposite(ESide InSide) { return InSide == ESide::Left ? ESide::Right : ESide::Left; }
+
+static TOptional<FCurveBounds> FindMinDistance(
+	const FKeyNeighborData& InData, ESide InSide, const FKeyHandle& InKey, const FCurveModel& InCurveModel
+	)
+{
+	const TOptional<double> HeightToUse = InData.Choose(InSide);
+	if (!HeightToUse)
+	{
+		return {};
+	}
+	
+	FKeyPosition CurrentKeyPosition{};
+	InCurveModel.GetKeyPositions({ InKey } , MakeArrayView(&CurrentKeyPosition, 1));
+	
+	const double CurrentY = CurrentKeyPosition.OutputValue;
+	return FCurveBounds{ FMath::Min(CurrentY, *HeightToUse), FMath::Max(CurrentY, *HeightToUse) };
+}
+
+static FCurveBounds FindDistanceToPointButDefaultToOtherPoint(
+	const FKeyNeighborData& InData, ESide InPreferredSide, const FKeyHandle& InKey, const FCurveModel& InCurveModel
+	)
+{
+	if (const TOptional<FCurveBounds> LeftBounds = FindMinDistance(InData, InPreferredSide, InKey, InCurveModel))
+	{
+		return *LeftBounds;
+	}
+	if (const TOptional<FCurveBounds> RightBounds = FindMinDistance(InData, Opposite(InPreferredSide), InKey, InCurveModel))
+	{
+		return *RightBounds;
+	}
+	return {};
+}
+
+/** @return Whether the center key is now on the other side of the neighbor than it was initially. */
+static bool HasCenterFlipped(const FKeyNeighborData& InData, ESide InSide, const FKeyHandle& InKey, const FCurveModel& InCurveModel)
+{
+	const TOptional<double> HeightToUse = InData.Choose(InSide);
+	if (!HeightToUse)
+	{
+		return false;
+	}
+	
+	FKeyPosition CurrentKeyPosition{};
+	InCurveModel.GetKeyPositions({ InKey } , MakeArrayView(&CurrentKeyPosition, 1));
+	
+	const double CurrentCenterHeight = CurrentKeyPosition.OutputValue;
+	const bool bCenterWasInitiallyAbove = InData.CenterInitialHeight >= *HeightToUse;
+	const bool bHasFlipped = (bCenterWasInitiallyAbove && CurrentCenterHeight < *HeightToUse) 
+		|| (!bCenterWasInitiallyAbove && CurrentCenterHeight > *HeightToUse);
+	return bHasFlipped;
+}
+
+static FCurveBounds ComputeDistanceToNeighbor(
+	const FKeyNeighborData& InNeighbourData, const FKeyHandle& InKey, const FCurveModel& InCurveModel, ESide InAnchorPoint
+	)
+{
+	TOptional<FKeyHandle> Before, After;
+	InCurveModel.GetNeighboringKeys(InKey, Before, After);
+	const TOptional<double> NeighbourHeight = InNeighbourData.Choose(InAnchorPoint);
+	if (!NeighbourHeight)
+	{
+		return FCurveBounds{};
+	}
+	
+	FKeyPosition CurrentPosition; 
+	InCurveModel.GetKeyPositions({ InKey }, MakeArrayView(&CurrentPosition, 1));
+	
+	const double CurrentY = CurrentPosition.OutputValue;
+	return FCurveBounds{ FMath::Min(CurrentY, *NeighbourHeight), FMath::Max(CurrentY, *NeighbourHeight) };
+}
+
+static TOptional<double> ComputeDeltaEdgeMovementForSingleKey(
+	const FSingleKeyMirrorData& InMirrorData, const FCurveModel& InCurveModel,
+	ESide InKeyToMirrorTo, ESide InAnchorPoint
+	)
+{
+	// The neighbor key we're trying to flatten the tangent to does not exist? Do nothing. 
+	const TOptional<double> TargetHeight = InMirrorData.NeighborData.Choose(InKeyToMirrorTo);
+	if (!TargetHeight)
+	{
+		return {};
+	}
+	
+	const TOptional<FCurveBounds> CurveBounds = ComputeDistanceToNeighbor(InMirrorData.NeighborData, InMirrorData.GetKeyHandle(), InCurveModel, InAnchorPoint);
+	if (!ensure(CurveBounds))
+	{
+		return {};
+	}
+	
+	const double DeltaEdge = CurveBounds->Delta();
+	// The user is interpolating towards the key that we set up the interpolation distance...
+	if (InKeyToMirrorTo == InAnchorPoint)
+	{
+		// ... yes: nothing to translate
+		return DeltaEdge;
+	}
+	
+	// ... no: the interpolation distance needs to be converted to the other key.
+	const double InitialEdgeDistance = InMirrorData.InitialDistanceToAnchorPoint;
+	const double AbsoluteMoved = FMath::Abs(DeltaEdge - InitialEdgeDistance);
+	
+	const double DeltaHeight = FMath::Abs(InMirrorData.NeighborData.CenterInitialHeight - *TargetHeight);
+	if (FMath::IsNearlyZero(DeltaHeight))
+	{
+		return {};
+	}
+	
+	const double MovedRelativeToRight = AbsoluteMoved / DeltaHeight;
+	const double MovedAbsToLeft = InitialEdgeDistance - MovedRelativeToRight * InitialEdgeDistance;
+	return MovedAbsToLeft;
+}
+}
+
+bool FTangentSingleKeyFlattener::ResetFromSelection(const FCurveEditor& InCurveEditor)
+{
+	CachedCurveData.Empty(CachedCurveData.Num());
+	for (const TPair<FCurveModelID, FKeyHandleSet>& Pair : InCurveEditor.GetSelection().GetAll())
+	{
+		AddTangents(InCurveEditor, Pair.Key, Pair.Value.AsArray());
+	}
+	return !CachedCurveData.IsEmpty();
+}
+
+bool FTangentSingleKeyFlattener::AddTangents(const FCurveEditor& InCurveEditor, const FCurveModelID& InCurveId, TConstArrayView<FKeyHandle> InKeys)
+{
+	const FCurveModel* CurveModel = InCurveEditor.FindCurve(InCurveId);
+	if (!CurveModel)
+	{
+		return false;
+	}
+		
+	FMirrorableTangentInfo TangentInfo = FilterMirrorableTangents(InCurveEditor, InCurveId, InKeys);
+	if (TangentInfo.MirrorableKeys.Num() != 1)
+	{
+		return false;
+	}
+	
+	const FKeyNeighborData KeyData = TangentFlattenDetail::AnalyzeKey(TangentInfo.MirrorableKeys[0], *CurveModel);
+	const TOptional<ESide> InitialAnchorPoint = TangentFlattenDetail::GetInitialAnchorSide(KeyData);
+	if (!InitialAnchorPoint)
+	{
+		return false;
+	}
+	
+	const FKeyHandle CenterKey = InKeys[0];
+	const auto[Min, Max] = TangentFlattenDetail::FindDistanceToPointButDefaultToOtherPoint(KeyData, *InitialAnchorPoint, CenterKey, *CurveModel);
+	const double DistanceToAnchorPoint = Max - Min;
+	if (FMath::IsNearlyZero(DistanceToAnchorPoint))
+	{
+		return false;
+	}
+
+	// We're going to have 2 fake edges: the height of the moved "edge" is the delta of min and max values in selection. The 2nd edge is 0.
+	// Effectively, as the selection is squished, the virtual edges move closer. The closer the virtual edges, the more the tangents are squished vertically.
+	// 0 height means tangent is 0.
+	constexpr double MidpointEdgeHeight = 0;
+	FCurveTangentMirrorData MirrorData(MoveTemp(TangentInfo), DistanceToAnchorPoint, MidpointEdgeHeight);
+	CachedCurveData.Add(InCurveId, FSingleKeyMirrorData(CenterKey, KeyData, MoveTemp(MirrorData), *InitialAnchorPoint, DistanceToAnchorPoint));
+	return true;
+}
+
+void FTangentSingleKeyFlattener::ComputeMirroringParallel(const FCurveEditor& InCurveEditor, ESide InKeyToMirrorTo)
+{
+	for (TPair<FCurveModelID, FSingleKeyMirrorData>& Pair : CachedCurveData)
+	{
+		FCurveModel* CurveModel = InCurveEditor.FindCurve(Pair.Key);
+		if (!CurveModel)
+		{
+			continue;
+		}
+		
+		FSingleKeyMirrorData& MirrorData = Pair.Value;
+		
+		// If the moved key has crossed the target key since, we need to mirror the tangent. In that case Alpha is in range [-1, 0].
+		const bool bFlip = TangentFlattenDetail::HasCenterFlipped(MirrorData.NeighborData, InKeyToMirrorTo, MirrorData.GetKeyHandle(), *CurveModel);
+		const double Sign = InKeyToMirrorTo == ESide::Right // Not quite sure why this is needed... need to investigate. Probably correct but cannot derive it right now.
+			? 1 
+			: bFlip ? -1 : 1;
+		
+		const TOptional<double> BlendHeight = TangentFlattenDetail::ComputeDeltaEdgeMovementForSingleKey(
+			MirrorData, *CurveModel, InKeyToMirrorTo, MirrorData.AnchorKeySide
+			);
+		if (BlendHeight)
+		{
+			RecomputeMirroringParallel(InCurveEditor, Pair.Key, MirrorData.TangentMirrorData, *BlendHeight * Sign);
+		}
+	}
+}
+}
\ No newline at end of file
diff --git a/Engine/Source/Editor/CurveEditor/Public/Misc/Mirror/MirrorUtils.h b/Engine/Source/Editor/CurveEditor/Public/Misc/Mirror/MirrorUtils.h
--- a/Engine/Source/Editor/CurveEditor/Public/Misc/Mirror/MirrorUtils.h
+++ b/Engine/Source/Editor/CurveEditor/Public/Misc/Mirror/MirrorUtils.h
@@ -30,6 +30,8 @@ struct FCurveBounds
 	double Min;
 	/** All keys are below this */
 	double Max;
+	
+	double Delta() const { return Max - Min; }
 };
 /** @return Find the min and max height values. */
 CURVEEDITOR_API FCurveBounds FindMinMaxHeight(
diff --git a/Engine/Source/Editor/CurveEditor/Public/Misc/Mirror/SingleKeyMirrorData.h b/Engine/Source/Editor/CurveEditor/Public/Misc/Mirror/SingleKeyMirrorData.h
new file mode 100644
--- /dev/null
+++ b/Engine/Source/Editor/CurveEditor/Public/Misc/Mirror/SingleKeyMirrorData.h
@@ -0,0 +1,74 @@
+﻿// Copyright Epic Games, Inc. All Rights Reserved.
+
+#pragma once
+
+#include "Containers/Map.h"
+#include "CurveEditor.h"
+#include "MultiCurveMirrorUtils.h"
+
+#define UE_API CURVEEDITOR_API
+
+class FCurveEditor;
+
+namespace UE::CurveEditor
+{
+enum class ESide { Left, Right };
+
+/**
+ * Information about a center key and its neighbor's heights.
+ * @see FTangentSingleKeyFlattener
+ */
+struct FKeyNeighborData
+{
+	/** Height of the key in the center. */
+	double CenterInitialHeight;
+	
+	/** Height of the key to the left of the center key. */
+	TOptional<double> LeftKeyHeight;
+	/** Height of the key to the right of the center key. */
+	TOptional<double> RightKeyHeight;
+	
+	/** @return The height of the neighbor specified. */
+	TOptional<double> Choose(ESide InSide) const { return InSide == ESide::Left ? LeftKeyHeight : RightKeyHeight; }
+	
+	/** @return Whether the center key has any neighbors. */
+	bool IsEmpty() const { return !LeftKeyHeight && !RightKeyHeight; }
+};
+
+/** 
+ * Information used to flatten tangent of a single key. 
+ * @see FTangentSingleKeyFlattener
+ * */
+struct FSingleKeyMirrorData
+{
+	/** The center key, which is being moved. */
+	FKeyHandle CenterKey;
+	/** Info about the key's neighbor hood. Used to compute the mirror value. */
+	FKeyNeighborData NeighborData;
+	
+	/** The data used to interpolate the key */
+	FCurveTangentMirrorData TangentMirrorData;
+	
+	/** Indicates the neighbor of the interpolated key with which the base interpolation distance initialized */
+	ESide AnchorKeySide;
+	/** The distance the center key had to the anchor key, when the data was initialized */
+	double InitialDistanceToAnchorPoint;
+	
+	/** @return The key whose tangent is being flattened */
+	FKeyHandle GetKeyHandle() const { return CenterKey; }
+
+	explicit FSingleKeyMirrorData(
+		FKeyHandle InCenterKey,
+		const FKeyNeighborData& InNeighborData, FCurveTangentMirrorData InTangentMirrorData, 
+		ESide InAnchorPoint, double InInitialDistanceToAnchorPoint
+		)
+		: CenterKey(InCenterKey)
+		, NeighborData(InNeighborData)
+		, TangentMirrorData(MoveTemp(InTangentMirrorData))
+		, AnchorKeySide(InAnchorPoint)
+		, InitialDistanceToAnchorPoint(InInitialDistanceToAnchorPoint)
+	{}
+};
+}
+
+#undef UE_API
diff --git a/Engine/Source/Editor/CurveEditor/Public/Misc/Mirror/TangentSingleKeyFlattener.h b/Engine/Source/Editor/CurveEditor/Public/Misc/Mirror/TangentSingleKeyFlattener.h
new file mode 100644
--- /dev/null
+++ b/Engine/Source/Editor/CurveEditor/Public/Misc/Mirror/TangentSingleKeyFlattener.h
@@ -0,0 +1,66 @@
+﻿// Copyright Epic Games, Inc. All Rights Reserved.
+
+#pragma once
+
+#include "Containers/Map.h"
+#include "CurveEditor.h"
+#include "SingleKeyMirrorData.h"
+
+#define UE_API CURVEEDITOR_API
+
+class FCurveEditor;
+
+namespace UE::CurveEditor
+{
+/**
+ * Flattens the tangents of a single key per curve whose tangent mode is user-specified (RCTM_User or RCTM_Break).
+ * The degree of flattening is proportional to how close the key moves toward a neighboring key vertically:
+ * reaching the neighbor results in fully flat (zero) tangents and passing the neighbor causes it to flip.
+ *
+ * === Algorithm for a center key k ===
+ * 
+ * AddTangents (call once to initialize):
+ *   1. Find the left and right neighbors of k.
+ *   2. Pick an anchor neighbor: prefer left, fall back to right.
+ *   3. Record initial_distance = |k.Y - anchor.Y|.
+ *   (No-ops if k has no neighbors, or if k and anchor share the same Y value.)
+ *   
+ * ComputeMirroringParallel (call each time k moves to recompute tangents):
+ *   1. Compute a blend height; a value that shrinks toward 0 as k approaches either neighbor:
+ *      - Moving toward anchor:   blend = |k.Y - anchor.Y|
+ *      - Moving toward opposite: blend = initial_distance * (1 - |current_distance_to_anchor - initial_distance| / |k_initial.Y - opposite.Y|)
+ *   2. If k has crossed past the target neighbor, flip the tangent direction.
+ *   3. Scale the tangents linearly by blend via RecomputeMirroringParallel.
+ *
+ * @see FTangentSelectionFlattener For flattening consecutive keys.
+ */
+class FTangentSingleKeyFlattener
+{
+public:
+
+	FTangentSingleKeyFlattener() = default;
+
+	/** Inits the tangent data from the curve editor's selection. */
+	UE_API bool ResetFromSelection(const FCurveEditor& InCurveEditor);
+
+	/** 
+	 * Searches the range for exactly 1 key that has user specified tangents and adds the key as blend info.
+	 * If there are more than 1 such keys, nothing happens.
+	 * @return Whether any data was added. 
+	 */
+	UE_API bool AddTangents(const FCurveEditor& InCurveEditor, const FCurveModelID& InCurveId, TConstArrayView<FKeyHandle> InKeys);
+	
+	/**
+	 * Recomputes the keys's tangents based on where on their vertical distance to the specified neighboring keys.
+	 * @param InCurveEditor Used to update key positions
+	 * @param InKeyToMirrorTo The neighboring key to use as pivot for mirror computation.
+	 */
+	UE_API void ComputeMirroringParallel(const FCurveEditor& InCurveEditor, ESide InKeyToMirrorTo);
+
+private:
+	
+	TMap<FCurveModelID, FSingleKeyMirrorData> CachedCurveData;
+};
+}
+
+#undef UE_API
