Curve Editor: Apply frame output padding for Normalized view (just like other views)

The padding was hardcoded to 10% (-0.1, 1.1). This now respects the frame output padding setting.

#Synced-CL 52681601

diff --git a/Engine/Source/Editor/CurveEditor/Private/Views/SCurveEditorViewNormalized.cpp b/Engine/Source/Editor/CurveEditor/Private/Views/SCurveEditorViewNormalized.cpp
--- a/Engine/Source/Editor/CurveEditor/Private/Views/SCurveEditorViewNormalized.cpp
+++ b/Engine/Source/Editor/CurveEditor/Private/Views/SCurveEditorViewNormalized.cpp
@@ -152,16 +152,26 @@ void SCurveEditorViewNormalized::FrameVertical(double InOutputMin, double InOutp
 {
 	if (!bFixedOutputBounds && InOutputMin < InOutputMax)
 	{
+		double OutputPadding = 0.1;
+		if (TSharedPtr<FCurveEditor> CurveEditor = GetCurveEditor())
+		{
+			const double PanelHeight = GetViewSpace().GetPhysicalHeight();
+			OutputPadding = PanelHeight != 0.0 ? FMath::Min((double)CurveEditor->GetSettings()->GetFrameOutputPadding() / PanelHeight, 0.5) : 0.1;
+		}
+
+		const double NormalizedMin = -OutputPadding;
+		const double NormalizedMax =  1.0 + OutputPadding;
+
 		if (AxisID)
 		{
 			FAxisInfo& AxisInfo = GetVerticalAxisInfo(AxisID);
-			AxisInfo.Min = -0.1;
-			AxisInfo.Max =  1.1;
+			AxisInfo.Min = NormalizedMin;
+			AxisInfo.Max = NormalizedMax;
 		}
 		else
 		{
-			OutputMin = -0.1;
-			OutputMax =  1.1;
+			OutputMin = NormalizedMin;
+			OutputMax = NormalizedMax;
 		}
 	}
 }
